[ActorModifier] Moved plugin

diff --git a/Engine/Plugins/Experimental/ActorModifierCore/Config/BaseActorModifierCore.ini b/Engine/Plugins/Experimental/ActorModifierCore/Config/BaseActorModifierCore.ini
new file mode 100644
--- /dev/null
+++ b/Engine/Plugins/Experimental/ActorModifierCore/Config/BaseActorModifierCore.ini
@@ -0,0 +1,9 @@
+﻿[CoreRedirects]
++ClassRedirects=(OldName="/Script/TagCollection.TagCollectionModifierSharedActor",NewName="/Script/ActorModifierCore.ActorModifierCoreSharedActor")
++ClassRedirects=(OldName="/Script/TagCollection.TagCollectionModifierBase",NewName="/Script/ActorModifierCore.ActorModifierCoreBase")
++ClassRedirects=(OldName="/Script/TagCollection.TagCollectionModifierComponent",NewName="/Script/ActorModifierCore.ActorModifierCoreComponent")
++ClassRedirects=(OldName="/Script/TagCollection.TagCollectionModifierSharedObject",NewName="/Script/ActorModifierCore.ActorModifierCoreSharedObject")
++ClassRedirects=(OldName="/Script/TagCollection.TagCollectionModifierStack",NewName="/Script/ActorModifierCore.ActorModifierCoreStack")
++ClassRedirects=(OldName="/Script/TagCollectionEditor.TagCollectionModifierStackCustomization",NewName="/Script/ActorModifierCoreEditor.ActorModifierCoreEditorStackCustomization")
++ClassRedirects=(OldName="/Script/TagCollectionEditor.TagCollectionModifierExtensionSubsystem",NewName="/Script/ActorModifierCoreEditor.ActorModifierCoreEditorSubsystem")
++StructRedirects=(OldName="/Script/TagCollectionEditor.TagCollectionModifierPropertiesWrapper",NewName="/Script/ActorModifierCoreEditor.ActorModifierCoreEditorPropertiesWrapper")
\ No newline at end of file
diff --git a/Engine/Plugins/Experimental/ActorModifierCore/Source/ActorModifierCore/Public/Modifiers/ActorModifierCoreComponent.h b/Engine/Plugins/Experimental/ActorModifierCore/Source/ActorModifierCore/Public/Modifiers/ActorModifierCoreComponent.h
new file mode 100644
--- /dev/null
+++ b/Engine/Plugins/Experimental/ActorModifierCore/Source/ActorModifierCore/Public/Modifiers/ActorModifierCoreComponent.h
@@ -0,0 +1,45 @@
+// Copyright Epic Games, Inc. All Rights Reserved.
+
+#pragma once
+
+#include "Components/ActorComponent.h"
+#include "ActorModifierCoreComponent.generated.h"
+
+class UActorModifierCoreStack;
+
+/** Component that contains a modifier stack and can be added on any actor */
+UCLASS(ClassGroup="Modifiers", BlueprintType, EditInlineNew, DefaultToInstanced, meta=(BlueprintSpawnableComponent), HideCategories=(Tags, Activation, Cooking, AssetUserData, Collision))
+class ACTORMODIFIERCORE_API UActorModifierCoreComponent : public UActorComponent
+{
+	GENERATED_BODY()
+
+public:
+	UActorModifierCoreComponent();
+	virtual ~UActorModifierCoreComponent() override;
+
+	UActorModifierCoreStack* GetModifierStack() const { return ModifierStack; }
+
+	/** Create this component for a specific actor and exposes it to the editor and for serialization */
+	static UActorModifierCoreComponent* CreateAndExposeComponent(AActor* InParentActor);
+
+	//~ Begin UActorComponent
+	virtual void OnComponentCreated() override;
+	virtual void OnComponentDestroyed(bool bDestroyingHierarchy) override;
+	//~ End UActorComponent
+
+	//~ Begin UObject
+#if WITH_EDITOR
+	virtual void PostEditUndo() override;
+#endif
+	//~ End UObject
+
+protected:
+	/** Initialize the inner root stack for the actor it is attached */
+	void InitializeStack();
+
+	virtual void TickComponent(float InDeltaTime, ELevelTick InTickType, FActorComponentTickFunction* InThisTickFunction) override;
+
+	/** This is the root stack that will contain all modifiers for this actor */
+	UPROPERTY(VisibleAnywhere, BlueprintReadOnly, NoClear, Export, Instanced, Category="Modifiers")
+	TObjectPtr<UActorModifierCoreStack> ModifierStack = nullptr;
+};
\ No newline at end of file
diff --git a/Engine/Plugins/Experimental/ActorModifierCore/Source/ActorModifierCore/Public/Subsystems/ActorModifierCoreSubsystem.h b/Engine/Plugins/Experimental/ActorModifierCore/Source/ActorModifierCore/Public/Subsystems/ActorModifierCoreSubsystem.h
new file mode 100644
--- /dev/null
+++ b/Engine/Plugins/Experimental/ActorModifierCore/Source/ActorModifierCore/Public/Subsystems/ActorModifierCoreSubsystem.h
@@ -0,0 +1,150 @@
+// Copyright Epic Games, Inc. All Rights Reserved.
+
+#pragma once
+
+#include "Modifiers/ActorModifierCoreDefs.h"
+#include "Subsystems/EngineSubsystem.h"
+#include "ActorModifierCoreSubsystem.generated.h"
+
+class AActor;
+class AActorModifierCoreSharedActor;
+class UClass;
+class UActorModifierCoreBase;
+class UActorModifierCoreSharedObject;
+class UActorModifierCoreStack;
+
+/** This subsystem handle all modifiers stack active in the engine and allows to create modifiers with registered metadata */
+UCLASS()
+class ACTORMODIFIERCORE_API UActorModifierCoreSubsystem : public UEngineSubsystem
+{
+	GENERATED_BODY()
+
+	friend class UActorModifierCoreEditorSubsystem;
+	friend class AActorModifierCoreSharedActor;
+	friend class UActorModifierCoreStack;
+	friend struct FActorModifierCoreMetadata;
+
+public:
+	DECLARE_MULTICAST_DELEGATE_OneParam(FOnModifierClassRegistered, const FActorModifierCoreMetadata& /* ModifierMetadata */)
+
+	/** Called when a modifier class is registered */
+	static FOnModifierClassRegistered OnModifierClassRegisteredDelegate;
+
+	/** Called when a modifier class is unregistered */
+	static FOnModifierClassRegistered OnModifierClassUnregisteredDelegate;
+
+	DECLARE_MULTICAST_DELEGATE_OneParam(FOnModifierStackRegistered, const UActorModifierCoreStack* /** ActorRootStack */)
+
+	/** Called when an actor modifier stack is registered */
+	static FOnModifierStackRegistered OnModifierStackRegisteredDelegate;
+
+	/** Called when an actor modifier stack is unregistered */
+	static FOnModifierStackRegistered OnModifierStackUnregisteredDelegate;
+
+	UActorModifierCoreSubsystem();
+
+	static UActorModifierCoreSubsystem* Get();
+
+	/** Default subsystem interface */
+	virtual void Initialize(FSubsystemCollectionBase& Collection) override;
+	virtual void Deinitialize() override;
+
+	/** Register a modifier class or override an already existing one */
+	bool RegisterModifierClass(const UClass* InModifierClass, bool bInOverrideIfExists = false);
+	bool UnregisterModifierClass(const FName& InName);
+	bool IsRegisteredModifierClass(const FName& InName) const;
+	bool IsRegisteredModifierClass(const UClass* InClass) const;
+
+	/** Return the modifier name or none from the modifier class provided */
+	FName GetRegisteredModifierName(const UClass* InModifierClass) const;
+
+	/** Return a set of class for all modifiers registered */
+	TSet<const UClass*> GetRegisteredModifierClasses() const;
+
+	/** Returns the name of modifiers that are currently registered */
+	TSet<FName> GetRegisteredModifiers() const;
+
+	/** Returns the modifiers names that are allowed for specific actor/stack and before another modifier */
+	TSet<FName> GetAllowedModifiers(AActor* InActor, UActorModifierCoreBase* InContextModifier = nullptr, EActorModifierCoreStackPosition InContextPosition = EActorModifierCoreStackPosition::Before) const;
+
+	/** Returns the modifiers name that match a specific category */
+	TSet<FName> GetCategoryModifiers(const FName& InCategory) const;
+
+	/** Returns the categories for all registered modifiers */
+	TSet<FName> GetModifierCategories() const;
+
+	/** Returns the category this modifier is in */
+	FName GetModifierCategory(const FName& InModifier) const;
+
+#if WITH_EDITOR
+	/** Returns the modifiers hidden to the user */
+	TSet<FName> GetHiddenModifiers() const;
+#endif
+
+	/** Returns modifiers from the stack where we can move the provided modifier, required and dependent modifiers from this MoveModifier will not appear in the list */
+	TArray<UActorModifierCoreBase*> GetAllowedMoveModifiers(UActorModifierCoreBase* InMoveModifier) const;
+
+	/** Register a root actor stack to query this actor stack from everywhere, is called automatically at creation or deserialization by stack */
+	bool RegisterActorModifierStack(UActorModifierCoreStack* InStack);
+
+	/** Unregister a root stack by providing the actor it is attached to */
+	bool UnregisterActorModifierStack(const AActor* InActor);
+
+	/** Unregister a root stack from an actor by providing the stack itself */
+	bool UnregisterActorModifierStack(const UActorModifierCoreStack* InStack);
+
+	/** Do we have a root modifier stack for this actor */
+	bool IsRegisteredActorModifierStack(const AActor* InActor) const;
+
+	/** Do we have a registered root stack */
+	bool IsRegisteredActorModifierStack(const UActorModifierCoreStack* InStack) const;
+
+	/** This is the correct way to retrieve a modifier stack for a specific actor if it has one */
+	UActorModifierCoreStack* GetActorModifierStack(const AActor* InActor) const;
+
+	/** This is the correct way to add modifier component to a specific actor if it has none */
+	UActorModifierCoreStack* AddActorModifierStack(AActor* InActor) const;
+
+	/** Return the actor this stack is linked to */
+	const AActor* GetModifierStackActor(const UActorModifierCoreStack* InStack) const;
+
+	/** Loops through each registered modifier metadata, only used to read */
+	bool ForEachModifierMetadata(TFunctionRef<bool(FActorModifierCoreMetadata&)> InProcessFunction) const;
+
+	/** Process this on matching registered metadata, only used to read */
+	bool ProcessModifierMetadata(const FName& InName, TFunctionRef<bool(FActorModifierCoreMetadata&)> InProcessFunction) const;
+
+	/** Get modifier shared object from a specific world and class, only one or none can exists per world */
+	UActorModifierCoreSharedObject* GetModifierSharedObject(UWorld* InWorld, TSubclassOf<UActorModifierCoreSharedObject> InClass, bool bInCreateIfNone = false) const;
+
+protected:
+	/** Checks whether we can add a modifier to the stack */
+	bool ValidateModifierCreation(const FName& InName, const UActorModifierCoreStack* InStack, FText& OutFailReason, UActorModifierCoreBase* InBeforeModifier = nullptr) const;
+
+	/** Creates a modifier based on a name and a stack, should be called by the stack itself */
+	UActorModifierCoreBase* CreateModifierInstance(const FName& InName, UActorModifierCoreStack* InStack, FText& OutFailReason, UActorModifierCoreBase* InBeforeModifier = nullptr) const;
+
+	/** Builds a modifier dependency list, does not make additional checks */
+	bool BuildModifierDependencies(const FName& InName, TArray<FName>& OutBeforeModifiers) const;
+
+	/** Done only once at subsystem initialization */
+	void ScanForModifiers();
+
+	/** Registers a modifier shared provider actor for a world */
+	bool RegisterModifierSharedProvider(AActorModifierCoreSharedActor* InSharedActor) const;
+
+	/** Unregisters a modifier shared provider actor for a world */
+	bool UnregisterModifierSharedProvider(const AActor* InSharedActor) const;
+
+	/** Gets a modifier shared provider for a world, creates one if none is found */
+	AActorModifierCoreSharedActor* GetModifierSharedProvider(UWorld* InWorld, bool bInSpawnIfNotFound = true) const;
+
+	/** stores modifiers metadata, you can override these if you want a different behaviour */
+	TMap<FName, TSharedRef<FActorModifierCoreMetadata>> ModifiersMetadata;
+
+	/** stores modifiers stacks per actor, there should be only one root stack per actor */
+	TMap<TWeakObjectPtr<const AActor>, TWeakObjectPtr<UActorModifierCoreStack>> ModifierStacks;
+
+	/** Stores modifiers providers per world, there should be only one provider per world */
+	TMap<TWeakObjectPtr<const UWorld>, TWeakObjectPtr<AActorModifierCoreSharedActor>> ModifierSharedProviders;
+};
diff --git a/Engine/Plugins/Experimental/ActorModifierCore/Source/ActorModifierCoreEditor/Private/ActorModifierCoreEditorModule.cpp b/Engine/Plugins/Experimental/ActorModifierCore/Source/ActorModifierCoreEditor/Private/ActorModifierCoreEditorModule.cpp
new file mode 100644
--- /dev/null
+++ b/Engine/Plugins/Experimental/ActorModifierCore/Source/ActorModifierCoreEditor/Private/ActorModifierCoreEditorModule.cpp
@@ -0,0 +1,41 @@
+// Copyright Epic Games, Inc. All Rights Reserved.
+
+#include "ActorModifierCoreEditorModule.h"
+
+#include "Modifiers/ActorModifierCoreStack.h"
+#include "Modifiers/Customizations/ActorModifierCoreEditorDetailCustomization.h"
+#include "Modules/ModuleManager.h"
+#include "PropertyEditorModule.h"
+
+void FActorModifierCoreEditorModule::StartupModule()
+{
+	RegisterCustomizations();
+}
+
+void FActorModifierCoreEditorModule::ShutdownModule()
+{
+	UnregisterCustomizations();
+}
+
+void FActorModifierCoreEditorModule::RegisterCustomizations() const
+{
+	// Register custom layouts
+	static FName PropertyEditor("PropertyEditor");
+	FPropertyEditorModule& PropertyModule = FModuleManager::GetModuleChecked<FPropertyEditorModule>(PropertyEditor);
+
+	PropertyModule.RegisterCustomClassLayout(UActorModifierCoreStack::StaticClass()->GetFName()
+		, FOnGetDetailCustomizationInstance::CreateStatic(&FActorModifierCoreEditorDetailCustomization::MakeInstance));
+}
+
+void FActorModifierCoreEditorModule::UnregisterCustomizations() const
+{
+	// Unregister custom layouts
+	static FName PropertyEditor("PropertyEditor");
+	if (FModuleManager::Get().IsModuleLoaded(PropertyEditor))
+	{
+		FPropertyEditorModule& PropertyModule = FModuleManager::GetModuleChecked<FPropertyEditorModule>(PropertyEditor);
+		PropertyModule.UnregisterCustomClassLayout(UActorModifierCoreStack::StaticClass()->GetFName());
+	}
+}
+
+IMPLEMENT_MODULE(FActorModifierCoreEditorModule, ActorModifierCoreEditor)
diff --git a/Engine/Plugins/Experimental/ActorModifierCore/Source/ActorModifierCoreEditor/Private/Modifiers/Customizations/ActorModifierCoreEditorStackCustomization.cpp b/Engine/Plugins/Experimental/ActorModifierCore/Source/ActorModifierCoreEditor/Private/Modifiers/Customizations/ActorModifierCoreEditorStackCustomization.cpp
new file mode 100644
--- /dev/null
+++ b/Engine/Plugins/Experimental/ActorModifierCore/Source/ActorModifierCoreEditor/Private/Modifiers/Customizations/ActorModifierCoreEditorStackCustomization.cpp
@@ -0,0 +1,918 @@
+// Copyright Epic Games, Inc. All Rights Reserved.
+
+#include "Modifiers/Customizations/ActorModifierCoreEditorStackCustomization.h"
+
+#include "Framework/Commands/GenericCommands.h"
+#include "Framework/Notifications/NotificationManager.h"
+#include "HAL/PlatformApplicationMisc.h"
+#include "IDetailTreeNode.h"
+#include "IPropertyRowGenerator.h"
+#include "JsonObjectConverter.h"
+#include "Modifiers/ActorModifierCoreBase.h"
+#include "Modifiers/ActorModifierCoreComponent.h"
+#include "Modifiers/ActorModifierCoreStack.h"
+#include "PropertyHandle.h"
+#include "ScopedTransaction.h"
+#include "Styling/SlateIconFinder.h"
+#include "Subsystems/ActorModifierCoreEditorSubsystem.h"
+#include "Subsystems/ActorModifierCoreSubsystem.h"
+#include "ActorModifierCoreEditorStyle.h"
+#include "Contexts/OperatorStackEditorMenuContext.h"
+#include "Items/OperatorStackEditorObjectItem.h"
+#include "Widgets/Notifications/SNotificationList.h"
+#include "Widgets/Views/STableRow.h"
+
+#define LOCTEXT_NAMESPACE "ActorModifierCoreEditorStackCustomization"
+
+DEFINE_LOG_CATEGORY_STATIC(LogActorModifierCoreEditorPropertiesWrapper, Log, All);
+
+UActorModifierCoreEditorStackCustomization::UActorModifierCoreEditorStackCustomization()
+	: UOperatorStackEditorStackCustomization(
+		TEXT("Modifiers")
+		, LOCTEXT("CustomizationLabel", "Modifiers")
+		, 1
+	)
+{
+	// for stack and modifiers
+	RegisterCustomizationFor(UActorModifierCoreBase::StaticClass());
+}
+
+bool UActorModifierCoreEditorStackCustomization::TransformContextItem(const FOperatorStackEditorItemPtr& InItem, TArray<FOperatorStackEditorItemPtr>& OutTransformedItems) const
+{
+	// If we have an actor then display stack
+	if (InItem->IsA<AActor>())
+	{
+		const UActorModifierCoreSubsystem* ModifierSubsystem = UActorModifierCoreSubsystem::Get();
+
+		if (UActorModifierCoreStack* Stack = ModifierSubsystem->GetActorModifierStack(InItem->Get<AActor>()))
+		{
+			OutTransformedItems.Add(MakeShared<FOperatorStackEditorObjectItem>(Stack));
+
+			return true;
+		}
+	}
+	// If we have a modifier component then display stack
+	else if (InItem->IsA<UActorModifierCoreComponent>())
+	{
+		const UActorModifierCoreComponent* ModifierComponent = InItem->Get<UActorModifierCoreComponent>();
+
+		OutTransformedItems.Add(MakeShared<FOperatorStackEditorObjectItem>(ModifierComponent->GetModifierStack()));
+
+		return true;
+	}
+	// If we have a stack then display modifiers inside
+	else if (InItem->IsA<UActorModifierCoreStack>())
+	{
+		const UActorModifierCoreStack* ModifierStack = InItem->Get<UActorModifierCoreStack>();
+
+		for (UActorModifierCoreBase* Modifier : ModifierStack->GetModifiers())
+		{
+			OutTransformedItems.Add(MakeShared<FOperatorStackEditorObjectItem>(Modifier));
+		}
+
+		return true;
+	}
+
+	return Super::TransformContextItem(InItem, OutTransformedItems);
+}
+
+void UActorModifierCoreEditorStackCustomization::CustomizeStackHeader(const FOperatorStackEditorTree& InItemTree, FOperatorStackEditorHeaderBuilder& InHeaderBuilder)
+{
+	if (!InItemTree.GetContext().GetItems().IsEmpty())
+	{
+		static const FName AddModifierMenuName = TEXT("AddModifierMenu");
+		if (!UToolMenus::Get()->IsMenuRegistered(AddModifierMenuName))
+		{
+			UToolMenu* const AddModifierMenu = UToolMenus::Get()->RegisterMenu(AddModifierMenuName, NAME_None, EMultiBoxType::Menu);
+			AddModifierMenu->AddDynamicSection(TEXT("PopulateAddModifierMenu"), FNewToolMenuDelegate::CreateUObject(this, &UActorModifierCoreEditorStackCustomization::FillStackHeaderMenu));
+		}
+
+		// Pin used categories
+		TSet<FString> PinnedKeywords;
+		for (FOperatorStackEditorItemPtr Item : InItemTree.GetAllItems())
+		{
+			if (Item.IsValid() && Item->IsA<UActorModifierCoreBase>())
+			{
+				const UActorModifierCoreBase* Modifier = Item->Get<UActorModifierCoreBase>();
+
+				if (Modifier && !Modifier->IsModifierStack())
+				{
+					PinnedKeywords.Add(Modifier->GetModifierMetadata().GetCategory().ToString());
+				}
+			}
+		}
+
+		InHeaderBuilder
+			.SetToolMenu(
+				AddModifierMenuName
+				, LOCTEXT("AddModifiersMenu", "Add Modifiers")
+				, FAppStyle::GetBrush(TEXT("Icons.PlusCircle"))
+			)
+			.SetSearchAllowed(true)
+			.SetSearchPinnedKeywords(PinnedKeywords);
+	}
+
+	Super::CustomizeStackHeader(InItemTree, InHeaderBuilder);
+}
+
+void UActorModifierCoreEditorStackCustomization::CustomizeItemHeader(const FOperatorStackEditorItemPtr& InItem, const FOperatorStackEditorTree& InItemTree, FOperatorStackEditorHeaderBuilder& InHeaderBuilder)
+{
+	const UActorModifierCoreSubsystem* ModifierSubsystem = UActorModifierCoreSubsystem::Get();
+
+	// Customize stack and modifier header
+	if (InItem->IsA<UActorModifierCoreBase>())
+	{
+		UActorModifierCoreBase* Modifier = InItem->Get<UActorModifierCoreBase>();
+
+		FBoolProperty* ModifierEnableProperty = FindFProperty<FBoolProperty>(UActorModifierCoreBase::StaticClass(), GET_MEMBER_NAME_CHECKED(UActorModifierCoreBase, bModifierEnabled));
+
+		const bool bIsStack = Modifier->IsModifierStack();
+
+		// Commands for item on key events
+		const TSharedPtr<FUICommandList> Commands = CreateModifierCommands(Modifier);
+
+		// Action menu available in header in slim toolbar
+		static const FName HeaderModifierMenuName = TEXT("HeaderModifierMenu");
+		if (!UToolMenus::Get()->IsMenuRegistered(HeaderModifierMenuName))
+		{
+			UToolMenu* const HeaderModifierMenu = UToolMenus::Get()->RegisterMenu(HeaderModifierMenuName, NAME_None, EMultiBoxType::SlimHorizontalToolBar);
+			HeaderModifierMenu->AddDynamicSection(TEXT("FillHeaderModifierMenu"), FNewToolMenuDelegate::CreateUObject(this, &UActorModifierCoreEditorStackCustomization::FillItemHeaderActionMenu));
+		}
+
+		// Context menu available when right clicking on item
+		static const FName ContextModifierMenuName = TEXT("ContextModifierMenu");
+		if (!UToolMenus::Get()->IsMenuRegistered(ContextModifierMenuName))
+		{
+			UToolMenu* const ContextModifierMenu = UToolMenus::Get()->RegisterMenu(ContextModifierMenuName, NAME_None, EMultiBoxType::Menu);
+			ContextModifierMenu->AddDynamicSection(TEXT("FillContextModifierMenu"), FNewToolMenuDelegate::CreateUObject(this, &UActorModifierCoreEditorStackCustomization::FillItemContextActionMenu));
+		}
+
+		// Item keyword for search
+		const TSet<FString> SearchKeywords
+		{
+			Modifier->GetModifierMetadata().GetDisplayName().ToString(),
+			Modifier->GetModifierMetadata().GetCategory().ToString()
+		};
+
+		FSlateIcon ModifierIcon = FSlateIconFinder::FindIconForClass(UActorModifierCoreBase::StaticClass());
+		FLinearColor ModifierColor = FLinearColor::Transparent;
+
+		ModifierSubsystem->ProcessModifierMetadata(Modifier->GetModifierName(), [&ModifierIcon, &ModifierColor](const FActorModifierCoreMetadata& InMetadata)
+		{
+			ModifierIcon = InMetadata.GetIcon();
+			ModifierColor = InMetadata.GetColor();
+			return true;
+		});
+
+		FString HeaderLabel = Modifier->GetModifierName().ToString();
+
+		// Add actor name next to stack label when we customize multiple items
+		if (bIsStack && InItemTree.GetRootItems().Num() > 1)
+		{
+			if (const AActor* ModifiedActor = Modifier->GetModifiedActor())
+			{
+				HeaderLabel += + TEXT(" (") + ModifiedActor->GetActorNameOrLabel() + TEXT(")");
+			}
+		}
+
+		/** Show last execution error messages if failed execution */
+		TAttribute<EOperatorStackEditorMessageType> MessageType = EOperatorStackEditorMessageType::None;
+		TAttribute<FText> MessageText = FText::GetEmpty();
+
+		if (!bIsStack)
+		{
+			TWeakObjectPtr<UActorModifierCoreBase> ModifierWeak(Modifier);
+
+			MessageType = TAttribute<EOperatorStackEditorMessageType>::CreateLambda([ModifierWeak]()
+			{
+				if (const UActorModifierCoreBase* Modifier = ModifierWeak.Get())
+				{
+					if (Modifier->GetModifierLastStatus().GetStatus() == EActorModifierCoreStatus::Warning)
+					{
+						return EOperatorStackEditorMessageType::Warning;
+					}
+
+					if (Modifier->GetModifierLastStatus().GetStatus() == EActorModifierCoreStatus::Error)
+					{
+						return EOperatorStackEditorMessageType::Error;
+					}
+
+					if (!Modifier->IsModifierEnabled())
+					{
+						return EOperatorStackEditorMessageType::Warning;
+					}
+				}
+
+				return EOperatorStackEditorMessageType::None;
+			});
+
+			MessageText = TAttribute<FText>::CreateLambda([ModifierWeak]()
+			{
+				if (const UActorModifierCoreBase* Modifier = ModifierWeak.Get())
+				{
+					if (!Modifier->IsModifierEnabled())
+					{
+						return LOCTEXT("ModifierDisabled", "Modifier disabled");
+					}
+
+					return Modifier->GetModifierLastStatus().GetStatusMessage();
+				}
+
+				return FText::GetEmpty();
+			});
+		}
+
+		InHeaderBuilder
+			.SetSearchAllowed(true)
+			.SetSearchKeywords(SearchKeywords)
+			.SetExpandable(!bIsStack)
+			.SetIcon(ModifierIcon.GetIcon())
+			.SetLabel(FText::FromString(HeaderLabel))
+			.SetBorderColor(ModifierColor)
+			.SetProperty(ModifierEnableProperty)
+			.SetCommandList(Commands)
+			.SetToolbarMenu(HeaderModifierMenuName)
+			.SetContextMenu(ContextModifierMenuName)
+			.SetMessageBox(MessageType, MessageText);
+	}
+
+	Super::CustomizeItemHeader(InItem, InItemTree, InHeaderBuilder);
+}
+
+void UActorModifierCoreEditorStackCustomization::CustomizeItemBody(const FOperatorStackEditorItemPtr& InItem, const FOperatorStackEditorTree& InItemTree, FOperatorStackEditorBodyBuilder& InBodyBuilder)
+{
+	// Customize stack and modifier body
+	if (InItem->IsA<UActorModifierCoreBase>())
+	{
+		FBoolProperty* ModifierEnableProperty = FindFProperty<FBoolProperty>(UActorModifierCoreBase::StaticClass(), GET_MEMBER_NAME_CHECKED(UActorModifierCoreBase, bModifierEnabled));
+		FBoolProperty* ProfilingEnableProperty = FindFProperty<FBoolProperty>(UActorModifierCoreStack::StaticClass(), GET_MEMBER_NAME_CHECKED(UActorModifierCoreStack, bModifierProfiling));
+
+		const bool bIsStack = InItem->IsA<UActorModifierCoreStack>();
+
+		InBodyBuilder
+			.SetShowDetailsView(!bIsStack)
+			.DisallowProperty(ModifierEnableProperty)
+			.DisallowProperty(ProfilingEnableProperty);
+	}
+
+	Super::CustomizeItemBody(InItem, InItemTree, InBodyBuilder);
+}
+
+void UActorModifierCoreEditorStackCustomization::CustomizeItemFooter(const FOperatorStackEditorItemPtr& InItem, const FOperatorStackEditorTree& InItemTree, FOperatorStackEditorFooterBuilder& InFooterBuilder)
+{
+	// Customize stack and modifier footer
+	if (InItem->IsA<UActorModifierCoreBase>())
+	{
+		const UActorModifierCoreBase* Modifier = InItem->Get<UActorModifierCoreBase>();
+		const TSharedPtr<FActorModifierCoreProfiler> Profiler = Modifier->GetProfiler();
+		UActorModifierCoreEditorSubsystem* ExtensionSubsystem = UActorModifierCoreEditorSubsystem::Get();
+
+		if (Profiler.IsValid())
+		{
+			const TSharedPtr<SWidget> Widget = ExtensionSubsystem->CreateProfilerWidget(Profiler);
+
+			InFooterBuilder
+				.SetCustomWidget(Widget);
+		}
+	}
+
+	Super::CustomizeItemFooter(InItem, InItemTree, InFooterBuilder);
+}
+
+bool UActorModifierCoreEditorStackCustomization::OnIsItemDraggable(const FOperatorStackEditorItemPtr& InDragItem)
+{
+	if (InDragItem->IsA<UActorModifierCoreBase>())
+	{
+		const bool bIsStack = InDragItem->IsA<UActorModifierCoreStack>();
+		return !bIsStack;
+	}
+
+	return Super::OnIsItemDraggable(InDragItem);
+}
+
+TOptional<EItemDropZone> UActorModifierCoreEditorStackCustomization::OnItemCanAcceptDrop(const TArray<FOperatorStackEditorItemPtr>& InDraggedItems, const FOperatorStackEditorItemPtr& InDropZoneItem, EItemDropZone InZone)
+{
+	const UActorModifierCoreEditorSubsystem* ExtensionSubsystem = UActorModifierCoreEditorSubsystem::Get();
+
+	TSet<UActorModifierCoreBase*> DraggedModifiers;
+	for (const FOperatorStackEditorItemPtr& Item : InDraggedItems)
+	{
+		if (!Item.IsValid())
+		{
+			continue;
+		}
+
+		if (!Item->IsA<UActorModifierCoreBase>())
+		{
+			continue;
+		}
+
+		if (UActorModifierCoreBase* Modifier = Item->Get<UActorModifierCoreBase>())
+		{
+			if (!Modifier->IsModifierStack())
+			{
+				DraggedModifiers.Add(Modifier);
+			}
+		}
+	}
+
+	if (InDropZoneItem->IsA<UActorModifierCoreBase>())
+	{
+		UActorModifierCoreBase* DropModifier = InDropZoneItem->Get<UActorModifierCoreBase>();
+
+		TArray<UActorModifierCoreBase*> MoveModifiers;
+		TArray<UActorModifierCoreBase*> CloneModifiers;
+
+		const EActorModifierCoreStackPosition Position = InZone == EItemDropZone::AboveItem ? EActorModifierCoreStackPosition::Before : EActorModifierCoreStackPosition::After;
+		ExtensionSubsystem->GetSortedModifiers(DraggedModifiers, DropModifier->GetModifiedActor(), DropModifier, Position, MoveModifiers, CloneModifiers);
+
+		if (!MoveModifiers.IsEmpty())
+		{
+			return InZone;
+		}
+	}
+
+	return Super::OnItemCanAcceptDrop(InDraggedItems, InDropZoneItem, InZone);
+}
+
+void UActorModifierCoreEditorStackCustomization::OnDropItem(const TArray<FOperatorStackEditorItemPtr>& InDraggedItems, const FOperatorStackEditorItemPtr& InDropZoneItem, EItemDropZone InZone)
+{
+	const UActorModifierCoreSubsystem* const ModifierSubsystem = UActorModifierCoreSubsystem::Get();
+	const UActorModifierCoreEditorSubsystem* ExtensionSubsystem = UActorModifierCoreEditorSubsystem::Get();
+
+	TSet<UActorModifierCoreBase*> Modifiers;
+	for (const FOperatorStackEditorItemPtr& Item : InDraggedItems)
+	{
+		if (!Item.IsValid())
+		{
+			continue;
+		}
+
+		if (!Item->IsA<UActorModifierCoreBase>())
+		{
+			continue;
+		}
+
+		if (UActorModifierCoreBase* Modifier = Item->Get<UActorModifierCoreBase>())
+		{
+			if (!Modifier->IsModifierStack())
+			{
+				Modifiers.Add(Modifier);
+			}
+		}
+	}
+
+	if (InDropZoneItem->IsA<UActorModifierCoreBase>())
+	{
+		UActorModifierCoreBase* DropModifier = InDropZoneItem->Get<UActorModifierCoreBase>();
+
+		const UActorModifierCoreStack* const TargetStack = DropModifier->GetModifierStack();
+
+		if (!IsValid(ModifierSubsystem) || !IsValid(TargetStack))
+		{
+			return;
+		}
+
+		AActor* const TargetActor = TargetStack->GetModifiedActor();
+
+		TArray<UActorModifierCoreBase*> MoveModifiers;
+		TArray<UActorModifierCoreBase*> CloneModifiers;
+
+		const EActorModifierCoreStackPosition Position = InZone == EItemDropZone::AboveItem ? EActorModifierCoreStackPosition::Before : EActorModifierCoreStackPosition::After;
+		ExtensionSubsystem->GetSortedModifiers(Modifiers, TargetActor, DropModifier, Position, MoveModifiers, CloneModifiers);
+
+		if (MoveModifiers.IsEmpty())
+		{
+			return;
+		}
+
+		FText FailReason;
+		ExtensionSubsystem->MoveModifiers(MoveModifiers, DropModifier, Position, &FailReason);
+
+		if (!FailReason.IsEmpty())
+		{
+			FNotificationInfo NotificationInfo(FailReason);
+			NotificationInfo.ExpireDuration = 3.f;
+			NotificationInfo.bFireAndForget = true;
+			FSlateNotificationManager::Get().AddNotification(NotificationInfo);
+		}
+	}
+
+	Super::OnDropItem(InDraggedItems, InDropZoneItem, InZone);
+}
+
+void UActorModifierCoreEditorStackCustomization::FillStackHeaderMenu(UToolMenu* InToolMenu) const
+{
+	if (!InToolMenu)
+	{
+		return;
+	}
+
+	UOperatorStackEditorMenuContext* const AddModifierContext = InToolMenu->FindContext<UOperatorStackEditorMenuContext>();
+	if (!AddModifierContext)
+	{
+		return;
+	}
+
+	FOperatorStackEditorContextPtr Context = AddModifierContext->GetContext();
+	if (!Context)
+	{
+		return;
+	}
+
+	const UActorModifierCoreEditorSubsystem* ModifierExtensionSubsystem = UActorModifierCoreEditorSubsystem::Get();
+	if (!IsValid(ModifierExtensionSubsystem))
+	{
+		return;
+	}
+
+	TSet<TWeakObjectPtr<UObject>> ContextObjects;
+	for (const FOperatorStackEditorItemPtr& ContextItem : Context->GetItems())
+	{
+		if (ContextItem->IsA<UObject>())
+		{
+			ContextObjects.Add(ContextItem->Get<UObject>());
+		}
+	}
+
+	const FActorModifierCoreEditorMenuContext MenuContext(ContextObjects);
+	FActorModifierCoreEditorMenuOptions MenuOptions(EActorModifierCoreEditorMenuType::Add);
+	MenuOptions.CreateSubMenu(false);
+	ModifierExtensionSubsystem->FillModifierMenu(InToolMenu, MenuContext, MenuOptions);
+}
+
+void UActorModifierCoreEditorStackCustomization::FillItemHeaderActionMenu(UToolMenu* InToolMenu) const
+{
+	if (!InToolMenu)
+	{
+		return;
+	}
+
+	const UOperatorStackEditorMenuContext* const MenuContext = InToolMenu->FindContext<UOperatorStackEditorMenuContext>();
+	if (!MenuContext)
+	{
+		return;
+	}
+
+	const FOperatorStackEditorItemPtr ItemContext = MenuContext->GetItem();
+	if (!ItemContext)
+	{
+		return;
+	}
+
+	UActorModifierCoreBase* Modifier = ItemContext->Get<UActorModifierCoreBase>();
+	if (!IsValid(Modifier))
+	{
+		return;
+	}
+
+	// Add profiling stat toggle entry
+	if (UActorModifierCoreStack* ModifierStack = Cast<UActorModifierCoreStack>(Modifier))
+	{
+		const FToolMenuEntry EnableProfilingModifierAction = FToolMenuEntry::InitToolBarButton(
+			TEXT("EnableProfilingModifierMenuEntry")
+			, FUIAction(
+				FExecuteAction::CreateUObject(this, &UActorModifierCoreEditorStackCustomization::ToggleModifierProfilingAction, ModifierStack)
+				, FCanExecuteAction()
+				, FIsActionChecked::CreateUObject(this, &UActorModifierCoreEditorStackCustomization::IsModifierProfiling, ModifierStack))
+			, FText::GetEmpty()
+			, FText::GetEmpty()
+			, FSlateIcon(FActorModifierCoreEditorStyle::Get().GetStyleSetName(), "Profiling")
+			, EUserInterfaceActionType::ToggleButton
+		);
+
+		InToolMenu->AddMenuEntry(EnableProfilingModifierAction.Name, EnableProfilingModifierAction);
+	}
+
+	// Add remove modifier entry
+	const FToolMenuEntry RemoveModifierAction = FToolMenuEntry::InitToolBarButton(
+		TEXT("RemoveModifierMenuEntry")
+		, FExecuteAction::CreateUObject(this, &UActorModifierCoreEditorStackCustomization::RemoveModifierAction, Modifier)
+		, FText::GetEmpty()
+		, FText::GetEmpty()
+		, FSlateIcon(FAppStyle::GetAppStyleSetName(), "Icons.Delete")
+	);
+
+	InToolMenu->AddMenuEntry(RemoveModifierAction.Name, RemoveModifierAction);
+}
+
+void UActorModifierCoreEditorStackCustomization::FillItemContextActionMenu(UToolMenu* InToolMenu) const
+{
+	if (!InToolMenu)
+	{
+		return;
+	}
+
+	const UOperatorStackEditorMenuContext* const MenuContext = InToolMenu->FindContext<UOperatorStackEditorMenuContext>();
+	if (!MenuContext)
+	{
+		return;
+	}
+
+	const FOperatorStackEditorItemPtr ItemContext = MenuContext->GetItem();
+	if (!ItemContext)
+	{
+		return;
+	}
+
+	UActorModifierCoreBase* Modifier = ItemContext->Get<UActorModifierCoreBase>();
+	if (!IsValid(Modifier))
+	{
+		return;
+	}
+
+	if (UActorModifierCoreStack* ModifierStack = Cast<UActorModifierCoreStack>(Modifier))
+	{
+		const FToolMenuEntry EnableProfilingModifierAction = FToolMenuEntry::InitMenuEntry(
+			TEXT("EnableProfilingModifierMenuEntry")
+			, LOCTEXT("EnableProfilingModifier", "Toggle profiling")
+			, FText::GetEmpty()
+			, FSlateIcon(FActorModifierCoreEditorStyle::Get().GetStyleSetName(), "Profiling")
+			, FUIAction(
+				FExecuteAction::CreateUObject(this, &UActorModifierCoreEditorStackCustomization::ToggleModifierProfilingAction, ModifierStack)
+				, FCanExecuteAction()
+				, FIsActionChecked::CreateUObject(this, &UActorModifierCoreEditorStackCustomization::IsModifierProfiling, ModifierStack))
+			, EUserInterfaceActionType::ToggleButton);
+
+		InToolMenu->AddMenuEntry(EnableProfilingModifierAction.Name, EnableProfilingModifierAction);
+	}
+
+	// Link menu entry to commands list for delete
+	{
+		TSharedPtr<const FUICommandList> Commands;
+		const TSharedPtr<FUICommandInfo> Command = FGenericCommands::Get().Delete;
+		InToolMenu->Context.GetActionForCommand(Command, Commands);
+		const FToolMenuEntry RemoveModifierMenuEntry = FToolMenuEntry::InitMenuEntryWithCommandList(Command, Commands);
+		InToolMenu->AddMenuEntry(RemoveModifierMenuEntry.Name, RemoveModifierMenuEntry);
+	}
+
+	// Link menu entry to commands list for copy
+	{
+		TSharedPtr<const FUICommandList> Commands;
+		const TSharedPtr<FUICommandInfo> Command = FGenericCommands::Get().Copy;
+		InToolMenu->Context.GetActionForCommand(Command, Commands);
+		const FToolMenuEntry CopyModifierMenuEntry = FToolMenuEntry::InitMenuEntryWithCommandList(Command, Commands);
+		InToolMenu->AddMenuEntry(CopyModifierMenuEntry.Name, CopyModifierMenuEntry);
+	}
+
+	// Link menu entry to commands list for paste
+	{
+		TSharedPtr<const FUICommandList> Commands;
+		const TSharedPtr<FUICommandInfo> Command = FGenericCommands::Get().Paste;
+		InToolMenu->Context.GetActionForCommand(Command, Commands);
+		const FToolMenuEntry PasteModifierMenuEntry = FToolMenuEntry::InitMenuEntryWithCommandList(Command, Commands);
+		InToolMenu->AddMenuEntry(PasteModifierMenuEntry.Name, PasteModifierMenuEntry);
+	}
+}
+
+bool UActorModifierCoreEditorStackCustomization::CanRemoveModifier(UActorModifierCoreBase* InModifier) const
+{
+	if (!IsValid(InModifier) || !IsValid(InModifier->GetModifiedActor()))
+	{
+		return false;
+	}
+
+	return true;
+}
+
+void UActorModifierCoreEditorStackCustomization::RemoveModifierAction(UActorModifierCoreBase* InModifier) const
+{
+	if (!CanRemoveModifier(InModifier))
+	{
+		return;
+	}
+
+	const UActorModifierCoreEditorSubsystem* ModifierEditorSubsystem = UActorModifierCoreEditorSubsystem::Get();
+
+	if (!IsValid(ModifierEditorSubsystem))
+	{
+		return;
+	}
+
+	if (InModifier->IsModifierStack())
+	{
+		TSet<AActor*> Actors
+		{
+			InModifier->GetModifiedActor()
+		};
+
+		ModifierEditorSubsystem->RemoveActorsModifiers(Actors);
+	}
+	else
+	{
+		TSet<UActorModifierCoreBase*> Modifiers;
+		Modifiers.Add(InModifier);
+
+		FText OutFailReason;
+		if (!ModifierEditorSubsystem->RemoveModifiers(Modifiers, &OutFailReason))
+		{
+			FNotificationInfo NotificationInfo(OutFailReason);
+			NotificationInfo.ExpireDuration = 3.0f;
+			NotificationInfo.bFireAndForget = true;
+			FSlateNotificationManager::Get().AddNotification(NotificationInfo);
+		}
+	}
+}
+
+bool UActorModifierCoreEditorStackCustomization::CanCopyModifier(UActorModifierCoreBase* InModifier) const
+{
+	if (!IsValid(InModifier) || InModifier->IsModifierStack())
+	{
+		return false;
+	}
+
+	return true;
+}
+
+void UActorModifierCoreEditorStackCustomization::CopyModifierAction(UActorModifierCoreBase* InModifier) const
+{
+	if (!CanCopyModifier(InModifier))
+	{
+		return;
+	}
+
+	TMap<FName, FString> ModifierPropertiesHandlesMap;
+	const bool bSuccess = CreatePropertiesHandlesMapFromModifier(InModifier, ModifierPropertiesHandlesMap);
+
+	if (!bSuccess)
+	{
+		return;
+	}
+
+	FString SerializedString;
+
+	const TSharedRef<FJsonObject> RootJsonObject = MakeShared<FJsonObject>();
+
+	FActorModifierCoreEditorPropertiesWrapper ModifierPropertiesWrapper;
+	ModifierPropertiesWrapper.ModifierName = InModifier->GetModifierName();
+	ModifierPropertiesWrapper.PropertiesHandlesAsStringMap = ModifierPropertiesHandlesMap;
+
+	TSharedRef<FJsonObject> PropertiesJsonObject = MakeShared<FJsonObject>();
+	FJsonObjectConverter::UStructToJsonObject(FActorModifierCoreEditorPropertiesWrapper::StaticStruct(), &ModifierPropertiesWrapper, PropertiesJsonObject, 0 /* CheckFlags */, 0 /* SkipFlags */);
+	const TSharedPtr<FJsonValue> PropertiesJsonValue = MakeShared<FJsonValueObject>(PropertiesJsonObject);
+
+	RootJsonObject->SetField(PropertiesWrapperEntry, PropertiesJsonValue);
+
+	const TSharedRef<TJsonWriter<>> Writer = TJsonWriterFactory<>::Create(&SerializedString);
+
+	FJsonSerializer::Serialize(RootJsonObject, Writer);
+
+	// Add Prefix to quickly identify whether current clipboard is from Tag Collection Properties or not
+	SerializedString = *FString::Printf(TEXT("%s%s")
+		, *PropertiesWrapperPrefix
+		, *SerializedString);
+
+	FPlatformApplicationMisc::ClipboardCopy(*SerializedString);
+}
+
+bool UActorModifierCoreEditorStackCustomization::CanPasteModifier(UActorModifierCoreBase* InModifier) const
+{
+	if (!IsValid(InModifier))
+	{
+		return false;
+	}
+
+	FString ClipboardContent;
+	FPlatformApplicationMisc::ClipboardPaste(ClipboardContent);
+	return ClipboardContent.StartsWith(PropertiesWrapperPrefix);
+}
+
+void UActorModifierCoreEditorStackCustomization::PasteModifierAction(UActorModifierCoreBase* InModifier) const
+{
+	if (!CanPasteModifier(InModifier))
+	{
+		return;
+	}
+
+	if (const UActorModifierCoreStack* ModifierStack = Cast<UActorModifierCoreStack>(InModifier))
+	{
+		TSet<AActor*> Actors { InModifier->GetModifiedActor() };
+		if (!AddModifierFromClipboard(Actors))
+		{
+			return;
+		}
+		InModifier = ModifierStack->GetLastModifier();
+	}
+
+	FActorModifierCoreEditorPropertiesWrapper ModifierPropertiesWrapper;
+	if (GetModifierPropertiesWrapperFromClipboard(ModifierPropertiesWrapper))
+	{
+		if (ModifierPropertiesWrapper.ModifierName == InModifier->GetModifierName())
+		{
+			FScopedTransaction Transaction(LOCTEXT("PasteModifierProperties", "Paste Modifier Properties"));
+			InModifier->Modify();
+
+			UpdateModifierFromPropertiesHandlesMap(InModifier, ModifierPropertiesWrapper.PropertiesHandlesAsStringMap);
+		}
+		else
+		{
+			const FString SourceModifierName = ModifierPropertiesWrapper.ModifierName.ToString();
+			const FString DestinationModifierName = InModifier->GetModifierName().ToString();
+
+			UE_LOG(LogActorModifierCoreEditorPropertiesWrapper, Warning, TEXT("Unable to copy properties from %s modifier to %s modifier"), *SourceModifierName, *DestinationModifierName);
+		}
+	}
+}
+
+bool UActorModifierCoreEditorStackCustomization::IsModifierProfiling(UActorModifierCoreStack* InStack) const
+{
+	if (!IsValid(InStack) || !IsValid(InStack->GetModifiedActor()))
+	{
+		return false;
+	}
+
+	return InStack->IsModifierProfiling();
+}
+
+void UActorModifierCoreEditorStackCustomization::ToggleModifierProfilingAction(UActorModifierCoreStack* InStack) const
+{
+	if (!IsValid(InStack) || !IsValid(InStack->GetModifiedActor()))
+	{
+		return;
+	}
+
+	InStack->SetModifierProfiling(!InStack->IsModifierProfiling());
+}
+
+TSharedRef<FUICommandList> UActorModifierCoreEditorStackCustomization::CreateModifierCommands(UActorModifierCoreBase* InModifier)
+{
+	TSharedRef<FUICommandList> Commands = MakeShared<FUICommandList>();
+
+	Commands->MapAction(FGenericCommands::Get().Copy
+		, FUIAction(
+			FExecuteAction::CreateUObject(this, &UActorModifierCoreEditorStackCustomization::CopyModifierAction, InModifier),
+			FCanExecuteAction::CreateUObject(this, &UActorModifierCoreEditorStackCustomization::CanCopyModifier, InModifier)
+		)
+	);
+
+	Commands->MapAction(FGenericCommands::Get().Paste
+		, FUIAction(
+			FExecuteAction::CreateUObject(this, &UActorModifierCoreEditorStackCustomization::PasteModifierAction, InModifier),
+			FCanExecuteAction::CreateUObject(this, &UActorModifierCoreEditorStackCustomization::CanPasteModifier, InModifier)
+		)
+	);
+
+	Commands->MapAction(FGenericCommands::Get().Delete
+		, FUIAction(
+			FExecuteAction::CreateUObject(this, &UActorModifierCoreEditorStackCustomization::RemoveModifierAction, InModifier),
+			FCanExecuteAction::CreateUObject(this, &UActorModifierCoreEditorStackCustomization::CanRemoveModifier, InModifier)
+		)
+	);
+
+	return Commands;
+}
+
+bool UActorModifierCoreEditorStackCustomization::CreatePropertiesHandlesMapFromModifier(UActorModifierCoreBase* InModifier, TMap<FName, FString>& OutModifierPropertiesHandlesMap) const
+{
+	if (!InModifier)
+	{
+		return false;
+	}
+
+	FPropertyEditorModule& PropertyEditor = FModuleManager::Get().LoadModuleChecked<FPropertyEditorModule>(TEXT("PropertyEditor"));
+
+	// no need to keep the generator outside of this function, since we are converting handles to string
+	const FPropertyRowGeneratorArgs RowGeneratorArgs;
+	const TSharedRef<IPropertyRowGenerator> PropertyRowGenerator = PropertyEditor.CreatePropertyRowGenerator(RowGeneratorArgs);
+
+	UObject* ModifierAsObj = Cast<UObject>(InModifier);
+	check(ModifierAsObj);
+
+	PropertyRowGenerator->SetObjects({ModifierAsObj});
+
+	const TArray<TSharedRef<IDetailTreeNode>>& RootTreeNodes = PropertyRowGenerator->GetRootTreeNodes();
+
+	for (const TSharedRef<IDetailTreeNode>& RootTreeNode : RootTreeNodes)
+	{
+		TArray<TSharedRef<IDetailTreeNode>> ChildrenNodes;
+		RootTreeNode->GetChildren(ChildrenNodes);
+
+		for (const TSharedRef<IDetailTreeNode>& Node : ChildrenNodes )
+		{
+			if (const TSharedPtr<IPropertyHandle>& PropertyHandle = Node->CreatePropertyHandle())
+			{
+				FName PropertyName = PropertyHandle->GetProperty()->NamePrivate;
+
+				FString PropertyValueAsString;
+				if (PropertyHandle->GetValueAsFormattedString(PropertyValueAsString, PPF_Copy) == FPropertyAccess::Success)
+				{
+					OutModifierPropertiesHandlesMap.Add(PropertyName, PropertyValueAsString);
+				}
+			}
+		}
+	}
+
+	return true;
+}
+
+bool UActorModifierCoreEditorStackCustomization::GetModifierPropertiesWrapperFromClipboard(FActorModifierCoreEditorPropertiesWrapper& OutPropertiesWrapper) const
+{
+	FString ClipboardContent;
+	FPlatformApplicationMisc::ClipboardPaste(ClipboardContent);
+
+	// remove prefix, this is not part of the modifier properties json data
+	ClipboardContent.RightChopInline(PropertiesWrapperPrefix.Len());
+
+	TSharedPtr<FJsonObject> RootJsonObject;
+	const TSharedRef<TJsonReader<>> Reader = TJsonReaderFactory<>::Create(ClipboardContent);
+
+	if (!FJsonSerializer::Deserialize(Reader, RootJsonObject))
+	{
+		UE_LOG(LogActorModifierCoreEditorPropertiesWrapper, Warning, TEXT("Unable to serialize the pasted text into Json format"));
+		return false;
+	}
+
+	const TSharedPtr<FJsonValue> PropertiesJsonValue = RootJsonObject->TryGetField(PropertiesWrapperEntry);
+
+	if (!PropertiesJsonValue.IsValid())
+	{
+		UE_LOG(LogActorModifierCoreEditorPropertiesWrapper, Warning, TEXT("Missing %s entry field in pasted text"), *PropertiesWrapperEntry);
+		return false;
+	}
+
+	if (PropertiesJsonValue->Type != EJson::Object)
+	{
+		UE_LOG(LogActorModifierCoreEditorPropertiesWrapper, Warning, TEXT("Invalid Modifier Properties Json Value. Not an object"));
+		return false;
+	}
+
+	const TSharedPtr<FJsonObject>& PropertiesJsonObject = PropertiesJsonValue->AsObject();
+	check(PropertiesJsonObject.IsValid());
+
+	const bool bStructSuccess = FJsonObjectConverter::JsonObjectToUStruct(PropertiesJsonObject.ToSharedRef(), FActorModifierCoreEditorPropertiesWrapper::StaticStruct(), &OutPropertiesWrapper, 0 /* CheckFlags */, 0 /* SkipFlags */);
+
+	return bStructSuccess;
+}
+
+bool UActorModifierCoreEditorStackCustomization::UpdateModifierFromPropertiesHandlesMap(UActorModifierCoreBase* InModifier, const TMap<FName, FString>& InModifierPropertiesHandlesMap) const
+{
+	if (!InModifier)
+	{
+		return false;
+	}
+
+	FPropertyEditorModule& PropertyEditor = FModuleManager::Get().LoadModuleChecked<FPropertyEditorModule>(TEXT("PropertyEditor"));
+
+	// no need to keep the generator outside of this function, since we are converting handles to string
+	const FPropertyRowGeneratorArgs RowGeneratorArgs;
+	const TSharedRef<IPropertyRowGenerator> PropertyRowGenerator = PropertyEditor.CreatePropertyRowGenerator(RowGeneratorArgs);
+
+	UObject* ModifierAsObj = Cast<UObject>(InModifier);
+	check(ModifierAsObj);
+
+	PropertyRowGenerator->SetObjects({ModifierAsObj});
+
+	const TArray<TSharedRef<IDetailTreeNode>>& RootTreeNodes = PropertyRowGenerator->GetRootTreeNodes();
+
+	for (const TSharedRef<IDetailTreeNode>& RootTreeNode : RootTreeNodes)
+	{
+		TArray<TSharedRef<IDetailTreeNode>> ChildrenNodes;
+		RootTreeNode->GetChildren(ChildrenNodes);
+
+		for (const TSharedRef<IDetailTreeNode>& Node : ChildrenNodes )
+		{
+			if (const TSharedPtr<IPropertyHandle>& PropertyHandle = Node->CreatePropertyHandle())
+			{
+				FName PropertyName = PropertyHandle->GetProperty()->NamePrivate;
+
+				if (InModifierPropertiesHandlesMap.Contains(PropertyName))
+				{
+					FString PropertyValueAsString = InModifierPropertiesHandlesMap[PropertyName];
+					PropertyHandle->SetValueFromFormattedString(PropertyValueAsString, EPropertyValueSetFlags::InstanceObjects);
+				}
+			}
+		}
+	}
+
+	return true;
+}
+
+bool UActorModifierCoreEditorStackCustomization::AddModifierFromClipboard(TSet<AActor*>& InActors) const
+{
+	const UActorModifierCoreEditorSubsystem* ModifierEditorSubsystem = UActorModifierCoreEditorSubsystem::Get();
+	if (!IsValid(ModifierEditorSubsystem))
+	{
+		return false;
+	}
+
+	FActorModifierCoreEditorPropertiesWrapper ModifierPropertiesWrapper;
+	if (GetModifierPropertiesWrapperFromClipboard(ModifierPropertiesWrapper))
+	{
+		const FName ModifierName = ModifierPropertiesWrapper.ModifierName;
+
+		FText OutFailReason;
+		if (!ModifierEditorSubsystem->AddActorsModifiers(ModifierName, InActors, &OutFailReason))
+		{
+			FNotificationInfo NotificationInfo(OutFailReason);
+			NotificationInfo.ExpireDuration = 3.0f;
+			NotificationInfo.bFireAndForget = true;
+			FSlateNotificationManager::Get().AddNotification(NotificationInfo);
+		}
+
+		return true;
+	}
+
+	return false;
+}
+
+#undef LOCTEXT_NAMESPACE
diff --git a/Engine/Plugins/Experimental/ActorModifierCore/Source/ActorModifierCoreEditor/Private/Subsystems/ActorModifierCoreEditorSubsystem.cpp b/Engine/Plugins/Experimental/ActorModifierCore/Source/ActorModifierCoreEditor/Private/Subsystems/ActorModifierCoreEditorSubsystem.cpp
new file mode 100644
--- /dev/null
+++ b/Engine/Plugins/Experimental/ActorModifierCore/Source/ActorModifierCoreEditor/Private/Subsystems/ActorModifierCoreEditorSubsystem.cpp
@@ -0,0 +1,842 @@
+// Copyright Epic Games, Inc. All Rights Reserved.
+
+#include "Subsystems/ActorModifierCoreEditorSubsystem.h"
+
+#include "Editor.h"
+#include "GameFramework/Actor.h"
+#include "Internationalization/Text.h"
+#include "Modifiers/ActorModifierCoreEditorMenu.h"
+#include "Modifiers/ActorModifierCoreStack.h"
+#include "Modifiers/Widgets/SActorModifierCoreEditorProfiler.h"
+#include "ScopedTransaction.h"
+#include "ActorModifierCoreEditorStyle.h"
+#include "Subsystems/ActorModifierCoreSubsystem.h"
+
+#define LOCTEXT_NAMESPACE "ActorModifierCoreEditorExtensionSubsystem"
+
+DEFINE_LOG_CATEGORY_STATIC(LogActorModifierCoreEditorSubsystem, Log, All);
+
+UActorModifierCoreEditorSubsystem::UActorModifierCoreEditorSubsystem()
+	: UEditorSubsystem()
+{
+}
+
+UActorModifierCoreEditorSubsystem* UActorModifierCoreEditorSubsystem::Get()
+{
+	if (GEditor)
+	{
+		return GEditor->GetEditorSubsystem<UActorModifierCoreEditorSubsystem>();
+	}
+	return nullptr;
+}
+
+void UActorModifierCoreEditorSubsystem::Initialize(FSubsystemCollectionBase& Collection)
+{
+	Super::Initialize(Collection);
+
+	EngineSubsystem = UActorModifierCoreSubsystem::Get();
+
+	check(EngineSubsystem.IsValid());
+	
+	// Initialize here if not called anywhere to setup ClassIcons
+	FActorModifierCoreEditorStyle::Get();
+}
+
+void UActorModifierCoreEditorSubsystem::Deinitialize()
+{
+	Super::Deinitialize();
+}
+
+bool UActorModifierCoreEditorSubsystem::EnableModifiers(bool bInEnabled, TSet<UActorModifierCoreBase*>& InModifiers, bool bInShouldTransact) const
+{
+	if (InModifiers.IsEmpty())
+	{
+		return false;
+	}
+
+	// create transaction
+	FText TransactionText;
+	if (bInEnabled)
+	{
+		TransactionText = FText::Format(
+			LOCTEXT("EnableModifiers.Enable", "Enabling {0} modifier(s)"),
+			FText::FromString(FString::FromInt(InModifiers.Num())));
+	}
+	else
+	{
+		TransactionText = FText::Format(
+			LOCTEXT("EnableModifiers.Disable", "Disabling {0} modifier(s)"),
+			FText::FromString(FString::FromInt(InModifiers.Num())));
+	}
+	FScopedTransaction Transaction(TransactionText, bInShouldTransact);
+	
+	// lets group modifiers to batch the operation, better to update the stack only once instead of many times
+	while (!InModifiers.IsEmpty())
+	{
+		TSet<UActorModifierCoreBase*> StackModifiers;
+
+		const UActorModifierCoreStack* CurrentStack = nullptr;
+	
+		for (UActorModifierCoreBase* Modifier : InModifiers)
+		{
+			if (!IsValid(Modifier))
+			{
+				continue;
+			}
+
+			UActorModifierCoreStack* ModifierStack = Modifier->GetModifierStack();
+			ModifierStack = IsValid(ModifierStack) ? ModifierStack : Modifier->GetRootModifierStack();
+			
+			if (!CurrentStack)
+			{
+				CurrentStack = ModifierStack;
+			}
+
+			if (ModifierStack == CurrentStack)
+			{
+				StackModifiers.Add(Modifier);
+			}
+		}
+
+		{
+			// disable stack modifiers update
+			FActorModifierCoreScopedLock Lock(StackModifiers);
+			for (UActorModifierCoreBase* Modifier : StackModifiers)
+			{
+				Modifier->SetModifierEnabled(bInEnabled);
+			}
+		}
+
+		InModifiers = InModifiers.Difference(StackModifiers);
+	}
+
+	return true;
+}
+
+bool UActorModifierCoreEditorSubsystem::RemoveModifiers(TSet<UActorModifierCoreBase*>& InModifiers, FText* OutFailReason, bool bInShouldTransact) const
+{
+	if (InModifiers.IsEmpty())
+	{
+		return false;
+	}
+
+	// create transaction
+	FText TransactionText;
+	const UActorModifierCoreBase* SingleModifier = InModifiers.Array()[0];
+	if (InModifiers.Num() == 1 && IsValid(SingleModifier))
+	{
+		TransactionText = FText::Format(
+			LOCTEXT("RemoveSingleModifiers", "Removing {0} modifier"),
+			FText::FromName(SingleModifier->GetModifierName()));
+	}
+	else
+	{
+		TransactionText = FText::Format(
+			LOCTEXT("RemoveMultipleModifiers", "Removing {0} modifier(s)"),
+			FText::FromString(FString::FromInt(InModifiers.Num())));
+	}
+	FScopedTransaction Transaction(TransactionText, bInShouldTransact);
+
+	// lets group modifiers to batch the operation, better to update the stack only once instead of many times
+	while (!InModifiers.IsEmpty())
+	{
+		TSet<UActorModifierCoreBase*> StackModifiers;
+
+		UActorModifierCoreStack* CurrentStack = nullptr;
+	
+		for (UActorModifierCoreBase* Modifier : InModifiers)
+		{
+			if (!IsValid(Modifier))
+			{
+				continue;
+			}
+
+			UActorModifierCoreStack* ModifierStack = Modifier->GetModifierStack();
+			ModifierStack = IsValid(ModifierStack) ? ModifierStack : Modifier->GetRootModifierStack();
+			
+			if (!CurrentStack)
+			{
+				CurrentStack = ModifierStack;
+			}
+
+			if (ModifierStack == CurrentStack)
+			{
+				StackModifiers.Add(Modifier);
+			}
+		}
+
+		// Sort by their order in stack first to avoid dependencies errors
+		bool bSuccess = true;
+		StackModifiers.Sort([CurrentStack](UActorModifierCoreBase& A, UActorModifierCoreBase& B)
+		{
+			bool bSortResult = true; 
+			CurrentStack->ProcessFunction([&A, &B, &bSortResult](const UActorModifierCoreBase* InModifier)->bool
+			{
+				if (InModifier == &A || InModifier == &B)
+				{
+					bSortResult = InModifier == &A ? false : true;
+					return false;
+				}
+				return true;
+			});
+			return bSortResult;
+		});
+
+		CurrentStack->ProcessLockFunction([this, &StackModifiers, OutFailReason, CurrentStack, &bSuccess]()
+		{
+			for (UActorModifierCoreBase* RemoveModifier : StackModifiers)
+			{
+				FActorModifierCoreStackRemoveOp RemoveOp;
+				RemoveOp.RemoveModifier = RemoveModifier;
+				RemoveOp.FailReason = OutFailReason;
+				if (!CurrentStack->RemoveModifier(RemoveOp))
+				{
+					bSuccess = false;
+					break;
+				}
+			}
+		});
+		
+		if (!bSuccess)
+		{
+			return false;
+		}
+
+		InModifiers = InModifiers.Difference(StackModifiers);
+	}
+
+	return true;
+}
+
+bool UActorModifierCoreEditorSubsystem::RemoveActorsModifiers(TSet<AActor*>& InActors, bool bInShouldTransact) const
+{
+	if (InActors.IsEmpty())
+	{
+		return false;
+	}
+
+	// Get actors stack
+	TArray<UActorModifierCoreStack*> ActorStacks;
+	for (const AActor* Actor : InActors)
+	{
+		if (!IsValid(Actor))
+		{
+			continue;
+		}
+
+		if (UActorModifierCoreStack* ActorStack = EngineSubsystem->GetActorModifierStack(Actor))
+		{
+			ActorStacks.Add(ActorStack);
+		}
+	}
+
+	// create transaction
+	const FText TransactionText = FText::Format(
+		LOCTEXT("AddModifier", "Remove all modifiers from {0} actor(s)"),
+		FText::FromString(FString::FromInt(InActors.Num())));
+	FScopedTransaction Transaction(TransactionText, bInShouldTransact);
+
+	// remove modifiers from actors
+	for (UActorModifierCoreStack* ActorStack : ActorStacks)
+	{
+		if (!ActorStack->RemoveAllModifiers())
+		{
+			return false;
+		}
+		// remove to indicate it was done
+		InActors.Remove(ActorStack->GetModifiedActor());
+	}
+	
+	return true;
+}
+
+bool UActorModifierCoreEditorSubsystem::AddActorsModifiers(const FName& InModifierName, TSet<AActor*>& InActors, FText* OutFailReason, bool bInShouldTransact) const
+{
+	if (!EngineSubsystem.IsValid())
+	{
+		return false;
+	}
+
+	if (!EngineSubsystem->IsRegisteredModifierClass(InModifierName))
+	{
+		return false;
+	}
+	
+	// Get actors stack, create if none
+	TArray<UActorModifierCoreStack*> ActorStacks;
+	for (AActor* Actor : InActors)
+	{
+		if (!IsValid(Actor))
+		{
+			continue;
+		}
+
+		UActorModifierCoreStack* ActorStack = EngineSubsystem->GetActorModifierStack(Actor);
+		
+		if (!ActorStack)
+		{
+			ActorStack = EngineSubsystem->AddActorModifierStack(Actor);
+		}
+
+		if (IsValid(ActorStack))
+		{
+			ActorStacks.Add(ActorStack);
+		}
+	}
+
+	// create transaction
+	const FText TransactionText = FText::Format(
+		LOCTEXT("AddModifier", "Add {0} modifier on {1} actor(s)"),
+		FText::FromName(InModifierName),
+		FText::FromString(FString::FromInt(InActors.Num())));
+	FScopedTransaction Transaction(TransactionText, bInShouldTransact);
+
+	// add modifier to actors
+	FActorModifierCoreStackInsertOp InsertOp;
+	InsertOp.NewModifierName = InModifierName;
+	InsertOp.FailReason = OutFailReason;
+	
+	for (UActorModifierCoreStack* ActorStack : ActorStacks)
+	{
+		if (!ActorStack->InsertModifier(InsertOp))
+		{
+			return false;
+		}
+		
+		// remove to indicate it was done
+		InActors.Remove(ActorStack->GetModifiedActor());
+	}
+	
+	return true;
+}
+
+bool UActorModifierCoreEditorSubsystem::InsertModifier(const FName& InModifierName, UActorModifierCoreStack* InStack, UActorModifierCoreBase* InPositionModifier, EActorModifierCoreStackPosition InPosition, FText* OutFailReason, bool bInShouldTransact) const
+{
+	if (!EngineSubsystem.IsValid())
+	{
+		return false;
+	}
+
+	if (!EngineSubsystem->IsRegisteredModifierClass(InModifierName))
+	{
+		return false;
+	}
+
+	if (!IsValid(InStack))
+	{
+		return false;
+	}
+
+	// create transaction
+	FText TransactionText;
+	if (InPosition == EActorModifierCoreStackPosition::Before)
+	{
+		if (IsValid(InPositionModifier))
+		{
+			TransactionText = FText::Format(
+				LOCTEXT("InsertModifier", "Insert {0} modifier before {1}"),
+				FText::FromName(InModifierName),
+				FText::FromName(InPositionModifier->GetModifierName()));
+		}
+		else
+		{
+			TransactionText = FText::Format(
+				LOCTEXT("InsertModifier", "Insert {0} modifier at the end of stack"),
+				FText::FromName(InModifierName));
+		}
+	}
+	else if (InPosition == EActorModifierCoreStackPosition::After)
+	{
+		if (IsValid(InPositionModifier))
+		{
+			TransactionText = FText::Format(
+				LOCTEXT("InsertModifier", "Insert {0} modifier after {1}"),
+				FText::FromName(InModifierName),
+				FText::FromName(InPositionModifier->GetModifierName()));
+		}
+		else
+		{
+			TransactionText = FText::Format(
+				LOCTEXT("InsertModifier", "Insert {0} modifier at the start of stack"),
+				FText::FromName(InModifierName));
+		}
+	}
+	FScopedTransaction Transaction(TransactionText, bInShouldTransact);
+
+	// insert modifier in stack
+	FActorModifierCoreStackInsertOp InsertOp;
+	InsertOp.NewModifierName = InModifierName;
+	InsertOp.InsertPosition = InPosition;
+	InsertOp.InsertPositionContext = InPositionModifier;
+	InsertOp.FailReason = OutFailReason;
+	return InStack->InsertModifier(InsertOp) != nullptr;
+}
+
+bool UActorModifierCoreEditorSubsystem::MoveModifier(UActorModifierCoreBase* InMoveModifier, UActorModifierCoreBase* InPositionModifier, EActorModifierCoreStackPosition InPosition, FText* OutFailReason, bool bInShouldTransact) const
+{
+	if (!EngineSubsystem.IsValid())
+	{
+		return false;
+	}
+
+	if (!IsValid(InMoveModifier))
+	{
+		return false;
+	}
+
+	UActorModifierCoreStack* ModifierStack = InMoveModifier->GetModifierStack();
+	if (!IsValid(ModifierStack))
+	{
+		return false;
+	}
+	
+	const FName& MoveModifierName = InMoveModifier->GetModifierName();
+	
+	// create transaction
+	FText TransactionText;
+	if (InPosition == EActorModifierCoreStackPosition::Before)
+	{
+		if (IsValid(InPositionModifier))
+		{
+			TransactionText = FText::Format(
+				LOCTEXT("InsertModifier", "Move {0} modifier before {1}"),
+				FText::FromName(MoveModifierName),
+				FText::FromName(InPositionModifier->GetModifierName()));
+		}
+		else
+		{
+			TransactionText = FText::Format(
+				LOCTEXT("InsertModifier", "Move {0} modifier at the end of stack"),
+				FText::FromName(MoveModifierName));
+		}
+	}
+	else if (InPosition == EActorModifierCoreStackPosition::After)
+	{
+		if (IsValid(InPositionModifier))
+		{
+			TransactionText = FText::Format(
+				LOCTEXT("InsertModifier", "Move {0} modifier after {1}"),
+				FText::FromName(MoveModifierName),
+				FText::FromName(InPositionModifier->GetModifierName()));
+		}
+		else
+		{
+			TransactionText = FText::Format(
+				LOCTEXT("InsertModifier", "Move {0} modifier at the start of stack"),
+				FText::FromName(MoveModifierName));
+		}
+	}
+	FScopedTransaction Transaction(TransactionText, bInShouldTransact);
+
+	// move modifier in stack
+	FActorModifierCoreStackMoveOp MoveOp;
+	MoveOp.MoveModifier = InMoveModifier;
+	MoveOp.MovePosition = InPosition;
+	MoveOp.MovePositionContext = InPositionModifier;
+	MoveOp.FailReason = OutFailReason;
+	return ModifierStack->MoveModifier(MoveOp);
+}
+
+bool UActorModifierCoreEditorSubsystem::FillModifierMenu(UToolMenu* InMenu, const FActorModifierCoreEditorMenuContext& InContext, const FActorModifierCoreEditorMenuOptions& InMenuOptions) const
+{
+	if (!InMenu || InContext.IsEmpty())
+	{
+		return false;
+	}
+	
+	const FActorModifierCoreEditorMenuData MenuData(InContext, InMenuOptions);
+	FToolMenuSection* ContextModifiersSection = nullptr;
+	if (InMenuOptions.ShouldCreateSubMenu())
+	{
+		static const FName ModifierSection("ContextModifierActions");
+	
+		ContextModifiersSection = InMenu->FindSection(ModifierSection);
+		if (!ContextModifiersSection)
+		{
+			ContextModifiersSection = &InMenu->AddSection(ModifierSection
+				, LOCTEXT("ContextModifierActions", "Modifiers Actions")
+				, FToolMenuInsert(NAME_None, EToolMenuInsertType::First));
+		}
+	}
+	
+	if (InMenuOptions.GetMenuType() == EActorModifierCoreEditorMenuType::Add)
+	{
+		if (InContext.ContainsAnyActor())
+		{
+			if (InMenuOptions.ShouldCreateSubMenu())
+			{
+				ContextModifiersSection->AddSubMenu(
+					TEXT("ContextAddModifiersActions"),
+					LOCTEXT("AddModifiers.Label", "Add Modifiers"),
+					LOCTEXT("AddModifiers.Tooltip", "Add modifiers to this selection"),
+					FNewToolMenuDelegate::CreateLambda(&UE::ActorModifierCoreEditor::OnExtendAddModifierMenu, MenuData));
+			}
+			else
+			{
+				UE::ActorModifierCoreEditor::OnExtendAddModifierMenu(InMenu, MenuData);
+			}
+			return true;
+		}
+	}
+	else if (InMenuOptions.GetMenuType() == EActorModifierCoreEditorMenuType::Delete)
+	{
+		if (InContext.ContainsNonEmptyStack() || InContext.ContainsAnyModifier())
+		{
+			if (InMenuOptions.ShouldCreateSubMenu())
+			{
+				ContextModifiersSection->AddSubMenu(
+					TEXT("ContextRemoveModifierActions"),
+					LOCTEXT("RemoveModifier.Label", "Remove Modifiers"),
+					LOCTEXT("EnableModifier.Tooltip", "Remove modifiers from this selection"),
+					FNewToolMenuDelegate::CreateLambda(&UE::ActorModifierCoreEditor::OnExtendRemoveModifierMenu, MenuData)
+				);
+			}
+			else
+			{
+				UE::ActorModifierCoreEditor::OnExtendRemoveModifierMenu(InMenu, MenuData);
+			}
+			return true;
+		}
+	}
+	else if (InMenuOptions.GetMenuType() == EActorModifierCoreEditorMenuType::Move)
+	{
+		if (InContext.ContainsOnlyModifier() && InContext.ContextModifiers.Num() == 1)
+		{
+			if (InMenuOptions.ShouldCreateSubMenu())
+			{
+				ContextModifiersSection->AddSubMenu(
+					TEXT("ContextMoveModifiersActions"),
+					LOCTEXT("MoveModifiers.Label", "Move Modifiers"),
+					LOCTEXT("MovModifiers.Tooltip", "Move modifiers in the stack"),
+					FNewToolMenuDelegate::CreateLambda(&UE::ActorModifierCoreEditor::OnExtendMoveModifierMenu, MenuData)
+				);
+			}
+			else
+			{
+				UE::ActorModifierCoreEditor::OnExtendMoveModifierMenu(InMenu, MenuData);
+			}
+			return true;
+		}
+	}
+	else if (InMenuOptions.GetMenuType() == EActorModifierCoreEditorMenuType::Enable)
+	{
+		if (InContext.ContainsDisabledModifier() || InContext.ContainsDisabledStack())
+		{
+			if (InMenuOptions.ShouldCreateSubMenu())
+			{
+				ContextModifiersSection->AddSubMenu(
+					TEXT("ContextEnableModifiersAction"),
+					LOCTEXT("EnableModifiers.Label", "Enable Modifier"),
+					LOCTEXT("EnableModifiers.Tooltip", "Enable modifier from this selection"),
+					FNewToolMenuDelegate::CreateLambda(&UE::ActorModifierCoreEditor::OnExtendEnableModifierMenu, MenuData)
+				);
+			}
+			else
+			{
+				UE::ActorModifierCoreEditor::OnExtendEnableModifierMenu(InMenu, MenuData);
+			}
+			return true;
+		}
+	}
+	else if (InMenuOptions.GetMenuType() == EActorModifierCoreEditorMenuType::Disable)
+	{
+		if (InContext.ContainsEnabledModifier() || InContext.ContainsEnabledStack())
+		{
+			if (InMenuOptions.ShouldCreateSubMenu())
+			{
+				ContextModifiersSection->AddSubMenu(
+					TEXT("ContextDisableModifiersAction"),
+					LOCTEXT("DisableModifiers.Label", "Disable Modifier"),
+					LOCTEXT("DisableModifiers.Tooltip", "Disable modifier from this selection"),
+					FNewToolMenuDelegate::CreateLambda(&UE::ActorModifierCoreEditor::OnExtendEnableModifierMenu, MenuData)
+				);
+			}
+			else
+			{
+				UE::ActorModifierCoreEditor::OnExtendEnableModifierMenu(InMenu, MenuData);
+			}
+			return true;
+		}
+	}
+	else if (InMenuOptions.GetMenuType() == EActorModifierCoreEditorMenuType::InsertBefore)
+	{
+		if (InContext.ContainsOnlyModifier() && InContext.ContextModifiers.Num() == 1)
+		{
+			if (InMenuOptions.ShouldCreateSubMenu())
+			{
+				ContextModifiersSection->AddSubMenu(
+					TEXT("ContextInsertBeforeModifiersActions"),
+					LOCTEXT("InsertBeforeModifiers.Label", "Insert Modifiers Before"),
+					LOCTEXT("InsertBeforeModifiers.Tooltip", "Insert modifiers before this selection"),
+					FNewToolMenuDelegate::CreateLambda(&UE::ActorModifierCoreEditor::OnExtendInsertModifierMenu, MenuData)
+				);
+			}
+			else
+			{
+				UE::ActorModifierCoreEditor::OnExtendInsertModifierMenu(InMenu, MenuData);
+			}
+			return true;
+		}
+	}
+	else if (InMenuOptions.GetMenuType() == EActorModifierCoreEditorMenuType::InsertAfter)
+	{
+		if (InContext.ContainsOnlyModifier() && InContext.ContextModifiers.Num() == 1)
+		{
+			if (InMenuOptions.ShouldCreateSubMenu())
+			{
+				ContextModifiersSection->AddSubMenu(
+					TEXT("ContextInsertAfterModifiersActions"),
+					LOCTEXT("InsertAfterModifiers.Label", "Insert Modifiers After"),
+					LOCTEXT("InsertAfterModifiers.Tooltip", "Insert modifiers after this selection"),
+					FNewToolMenuDelegate::CreateLambda(&UE::ActorModifierCoreEditor::OnExtendInsertModifierMenu, MenuData)
+				);
+			}
+			else
+			{
+				UE::ActorModifierCoreEditor::OnExtendInsertModifierMenu(InMenu, MenuData);
+			}
+			return true;
+		}
+	}
+
+	return false;
+}
+
+void UActorModifierCoreEditorSubsystem::GetSortedModifiers(const TSet<UActorModifierCoreBase*>& InModifiers, AActor* InTargetActor, UActorModifierCoreBase* InTargetModifier, EActorModifierCoreStackPosition InPosition, TArray<UActorModifierCoreBase*>& OutMoveModifiers, TArray<UActorModifierCoreBase*>& OutCloneModifiers) const
+{
+	const UActorModifierCoreSubsystem* const ModifierSubsystem = UActorModifierCoreSubsystem::Get();
+	
+	if (!IsValid(ModifierSubsystem) || InModifiers.IsEmpty() || !IsValid(InTargetActor))
+	{
+		return;
+	}
+
+	for (UActorModifierCoreBase* Modifier : InModifiers)
+	{
+		if (!IsValid(Modifier))
+		{
+			continue;	
+		}
+
+		// it's a clone operation if target actor is different as modifier actor
+		if (Modifier->GetModifiedActor() != InTargetActor)
+		{
+			OutCloneModifiers.Add(Modifier);
+		}
+		// it's a move operation if target actor is same as modifier actor
+		else
+		{
+			OutMoveModifiers.Add(Modifier);
+		}
+	}
+	
+	// Remove unsupported modifiers when inserting where target modifier is nullptr
+	if (!InTargetModifier)
+	{
+		const TSet<FName> AllowedModifiers = ModifierSubsystem->GetAllowedModifiers(InTargetActor, InTargetModifier, InPosition);
+		
+		OutMoveModifiers.RemoveAll([&AllowedModifiers](UActorModifierCoreBase* InModifier)
+		{
+			return !IsValid(InModifier) || !AllowedModifiers.Contains(InModifier->GetModifierName());
+		});
+
+		OutCloneModifiers.RemoveAll([&AllowedModifiers](UActorModifierCoreBase* InModifier)
+		{
+			return !IsValid(InModifier) || !AllowedModifiers.Contains(InModifier->GetModifierName());
+		});
+	}
+
+	// Sort them by dependency and current order
+	OutMoveModifiers.Sort([InTargetModifier](const UActorModifierCoreBase& A, const UActorModifierCoreBase& B)
+	{
+		const FActorModifierCoreMetadata& MetadataA = A.GetModifierMetadata();
+		const FActorModifierCoreMetadata& MetadataB = B.GetModifierMetadata();
+		const bool bADependsOnB = MetadataA.DependsOn(MetadataB.GetName());
+
+		const UActorModifierCoreStack* StackA = A.GetModifierStack();
+		const UActorModifierCoreStack* StackB = B.GetModifierStack();
+		const bool bAIsBeforeB = StackA == StackB && StackA->ContainsModifierBefore(&A, &B);
+
+		bool bSortOrder = bAIsBeforeB && !bADependsOnB;
+		
+		const UActorModifierCoreBase* DependencyModifier = bADependsOnB ? &A : &B;
+		if (StackA->ContainsModifierAfter(InTargetModifier, DependencyModifier))
+		{
+			bSortOrder = !bSortOrder;
+		}
+			
+		return bSortOrder;
+	});
+
+	OutCloneModifiers.Sort([](const UActorModifierCoreBase& A, const UActorModifierCoreBase& B)
+	{
+		const FActorModifierCoreMetadata& MetadataA = A.GetModifierMetadata();
+		const FActorModifierCoreMetadata& MetadataB = B.GetModifierMetadata();
+		const bool bADependsOnB = MetadataA.DependsOn(MetadataB.GetName());
+
+		const UActorModifierCoreStack* StackA = A.GetModifierStack();
+		const UActorModifierCoreStack* StackB = B.GetModifierStack();
+		const bool bAIsBeforeB = StackA == StackB && StackA->ContainsModifierBefore(&A, &B);
+
+		const bool bSortOrder = bAIsBeforeB && !bADependsOnB;
+
+		return bSortOrder;
+	});
+}
+
+bool UActorModifierCoreEditorSubsystem::MoveModifiers(const TArray<UActorModifierCoreBase*>& InModifiers, UActorModifierCoreBase* InTargetModifier, EActorModifierCoreStackPosition InPosition, FText* OutFailReason, bool bInShouldTransact) const
+{
+	if (!InTargetModifier)
+	{
+		return false;
+	}
+	
+	const UActorModifierCoreSubsystem* const ModifierSubsystem = UActorModifierCoreSubsystem::Get();
+	UActorModifierCoreStack* const TargetStack = InTargetModifier->GetModifierStack();
+
+	if (!IsValid(ModifierSubsystem) || !IsValid(TargetStack) || InModifiers.IsEmpty())
+	{
+		return false;
+	}
+
+	static const FText TransactionText = LOCTEXT("MoveModifiers", "Moving {0} modifier(s) {1} modifier {2}");
+	FScopedTransaction Transaction(
+		FText::Format(TransactionText
+			, FText::FromString(FString::FromInt(InModifiers.Num()))
+			, FText::FromString(InPosition == EActorModifierCoreStackPosition::After ? TEXT("after") : TEXT("before"))
+			, FText::FromName(InTargetModifier->GetModifierName()))
+		, bInShouldTransact);
+	
+	uint32 EditModifierCount = 0;
+
+	TargetStack->ProcessLockFunction([TargetStack, &InModifiers, OutFailReason, InPosition, InTargetModifier, &EditModifierCount]()
+	{
+		UActorModifierCoreBase* OperationContext = nullptr;
+		
+		for (UActorModifierCoreBase* Modifier : InModifiers)
+		{
+			if (IsValid(Modifier))
+			{
+				FActorModifierCoreStackMoveOp MoveOp;
+				MoveOp.FailReason = OutFailReason;
+				MoveOp.MovePosition = InPosition;
+				MoveOp.MovePositionContext = InTargetModifier;
+				MoveOp.MoveModifier = Modifier;
+
+				// Eg 1: when moving [bend, subdivide] before target modifier X : X is after them in the stack, we want bend added before X and subdivide added before Bend
+				// Eg 2: when moving [subdivide, bend] after target modifier X : X is before them in the stack, we want subdivide after X and bend after subdivide
+				if (OperationContext &&
+					((MoveOp.MovePosition == EActorModifierCoreStackPosition::After && TargetStack->ContainsModifierAfter(Modifier, InTargetModifier)) ||
+					(MoveOp.MovePosition == EActorModifierCoreStackPosition::Before && TargetStack->ContainsModifierBefore(Modifier, InTargetModifier))))
+				{
+					MoveOp.MovePositionContext = OperationContext;
+				}
+				
+				if (TargetStack->MoveModifier(MoveOp))
+				{
+					++EditModifierCount;
+					OperationContext = Modifier;
+				}
+				else
+				{
+					const AActor* TargetActor = InTargetModifier->GetModifiedActor();
+					const FText& ErrorText = MoveOp.FailReason ? *MoveOp.FailReason : FText::GetEmpty();
+					
+					// Move modifiers on actor failed
+					UE_LOG(LogActorModifierCoreEditorSubsystem, Warning, TEXT("Move modifier %s on actor %s failed : %s"),
+							*MoveOp.MoveModifier->GetModifierName().ToString(),
+							*TargetActor->GetActorNameOrLabel(),
+							*ErrorText.ToString());
+					break;
+				}
+			}
+		}
+	});
+
+	return EditModifierCount > 0;
+}
+
+bool UActorModifierCoreEditorSubsystem::CloneModifiers(const TArray<UActorModifierCoreBase*>& InModifiers, UActorModifierCoreBase* InTargetModifier, EActorModifierCoreStackPosition InPosition, FText* OutFailReason, bool bInShouldTransact) const
+{
+	if (!InTargetModifier)
+	{
+		return false;
+	}
+	
+	const UActorModifierCoreSubsystem* const ModifierSubsystem = UActorModifierCoreSubsystem::Get();
+	UActorModifierCoreStack* const TargetStack = InTargetModifier->GetModifierStack();
+
+	if (!IsValid(ModifierSubsystem) || !IsValid(TargetStack) || InModifiers.IsEmpty())
+	{
+		return false;
+	}
+
+	static const FText TransactionText = LOCTEXT("CloneModifiers", "Cloning {0} modifier(s) {1} modifier {2}");
+	FScopedTransaction Transaction(
+		FText::Format(TransactionText
+			, FText::FromString(FString::FromInt(InModifiers.Num()))
+			, FText::FromString(InPosition == EActorModifierCoreStackPosition::After ? TEXT("after") : TEXT("before"))
+			, FText::FromName(InTargetModifier->GetModifierName()))
+		, bInShouldTransact);
+	
+	uint32 EditModifierCount = 0;
+
+	TargetStack->ProcessLockFunction([TargetStack, &InModifiers, OutFailReason, InPosition, InTargetModifier, &EditModifierCount]()
+	{
+		UActorModifierCoreBase* OperationContext = nullptr;
+		for (UActorModifierCoreBase* Modifier : InModifiers)
+		{
+			if (IsValid(Modifier))
+			{
+				FActorModifierCoreStackCloneOp CloneOp;
+				CloneOp.FailReason = OutFailReason;
+				CloneOp.ClonePosition = InPosition;
+				CloneOp.ClonePositionContext = InTargetModifier;
+				CloneOp.CloneModifier = Modifier;
+
+				// When inserting A, B where B depends on A after C, clone A after C then clone B after A
+				if (OperationContext && CloneOp.ClonePosition == EActorModifierCoreStackPosition::After)
+				{
+					CloneOp.ClonePositionContext = OperationContext;
+				}
+				
+				if (UActorModifierCoreBase* NewClonedModifier = TargetStack->CloneModifier(CloneOp))
+				{
+					++EditModifierCount;
+					OperationContext = NewClonedModifier;
+				}
+				else
+				{
+					const AActor* TargetActor = InTargetModifier->GetModifiedActor();
+					const FText& ErrorText = CloneOp.FailReason ? *CloneOp.FailReason : FText::GetEmpty();
+					
+					// Clone modifiers on actor failed
+					UE_LOG(LogActorModifierCoreEditorSubsystem, Warning, TEXT("Clone modifier %s on actor %s failed : %s"),
+							*CloneOp.CloneModifier->GetModifierName().ToString(),
+							*TargetActor->GetActorNameOrLabel(),
+							*ErrorText.ToString());
+					break;
+				}
+			}
+		}
+	});
+
+	return EditModifierCount > 0;
+}
+
+TSharedPtr<SActorModifierCoreEditorProfiler> UActorModifierCoreEditorSubsystem::CreateProfilerWidget(TSharedPtr<FActorModifierCoreProfiler> InProfiler)
+{
+	if (!InProfiler.IsValid())
+	{
+		return nullptr;
+	}
+
+	if (const TFunction<TSharedRef<SActorModifierCoreEditorProfiler>(TSharedPtr<FActorModifierCoreProfiler>)>* WidgetFunction = ModifierProfilerWidgets.Find(InProfiler->GetProfilerType()))
+	{
+		return (*WidgetFunction)(InProfiler);
+	}
+
+	return SNew(SActorModifierCoreEditorProfiler, InProfiler);
+}
+
+#undef LOCTEXT_NAMESPACE
diff --git a/Engine/Plugins/Experimental/ActorModifierCore/Source/ActorModifierCoreEditor/Public/Subsystems/ActorModifierCoreEditorSubsystem.h b/Engine/Plugins/Experimental/ActorModifierCore/Source/ActorModifierCoreEditor/Public/Subsystems/ActorModifierCoreEditorSubsystem.h
new file mode 100644
--- /dev/null
+++ b/Engine/Plugins/Experimental/ActorModifierCore/Source/ActorModifierCoreEditor/Public/Subsystems/ActorModifierCoreEditorSubsystem.h
@@ -0,0 +1,100 @@
+// Copyright Epic Games, Inc. All Rights Reserved.
+
+#pragma once
+
+#include "EditorSubsystem.h"
+#include "Modifiers/ActorModifierCoreBase.h"
+#include "Modifiers/ActorModifierCoreDefs.h"
+#include "Modifiers/ActorModifierCoreEditorMenuDefs.h"
+#include "Modifiers/Widgets/SActorModifierCoreEditorProfiler.h"
+#include "ToolMenus.h"
+#include "UObject/WeakObjectPtrTemplates.h"
+#include "ActorModifierCoreEditorSubsystem.generated.h"
+
+class AActor;
+class FActorModifierCoreProfiler;
+class UActorModifierCoreBase;
+class UActorModifierCoreStack;
+class UActorModifierCoreSubsystem;
+
+/** Singleton class that handles editor operations for modifiers */
+UCLASS()
+class ACTORMODIFIERCOREEDITOR_API UActorModifierCoreEditorSubsystem : public UEditorSubsystem
+{
+	GENERATED_BODY()
+
+	friend class SActorModifierCoreEditorProfiler;
+	
+public:
+	UActorModifierCoreEditorSubsystem();
+	
+	static UActorModifierCoreEditorSubsystem* Get();
+	
+	//~ Begin USubsystem
+	virtual void Initialize(FSubsystemCollectionBase& Collection) override;
+	virtual void Deinitialize() override;
+	//~ End USubsystem
+
+	/** Enable or disable modifiers, will update the original array and perform a transaction if wanted */
+	bool EnableModifiers(bool bInEnabled, TSet<UActorModifierCoreBase*>& InModifiers, bool bInShouldTransact = true) const;
+	
+	/** Remove modifiers from different actors or stacks, will update the original array and perform a transaction if wanted */
+	bool RemoveModifiers(TSet<UActorModifierCoreBase*>& InModifiers, FText* OutFailReason = nullptr, bool bInShouldTransact = true) const;
+
+	/** Remove all modifiers from actors, will update the original array and perform a transaction if wanted */
+	bool RemoveActorsModifiers(TSet<AActor*>& InActors, bool bInShouldTransact = true) const;
+
+	/** Add a modifier to actors, adds a stack automatically if none is found, will update the original array and perform a transaction if wanted */
+	bool AddActorsModifiers(const FName& InModifierName, TSet<AActor*>& InActors, FText* OutFailReason = nullptr, bool bInShouldTransact = true) const;
+
+	/** Insert a modifier in a stack before or after another modifier, will perform a transaction if wanted */
+	bool InsertModifier(const FName& InModifierName, UActorModifierCoreStack* InStack, UActorModifierCoreBase* InPositionModifier, EActorModifierCoreStackPosition InPosition, FText* OutFailReason = nullptr, bool bInShouldTransact = true) const;
+
+	/** Moves a modifier in a stack before or after another modifier, will perform a transaction if wanted */
+	bool MoveModifier(UActorModifierCoreBase* InMoveModifier, UActorModifierCoreBase* InPositionModifier, EActorModifierCoreStackPosition InPosition, FText* OutFailReason = nullptr, bool bInShouldTransact = true) const;
+
+	/** Fills a menu based on context objects and menu options, will perform menu action transaction if wanted, returns true if the menu has been modified */
+	bool FillModifierMenu(UToolMenu* InMenu, const FActorModifierCoreEditorMenuContext& InContext, const FActorModifierCoreEditorMenuOptions& InMenuOptions) const;
+
+	/** Get order dependent modifiers in the correct order for stack operations */
+	void GetSortedModifiers(const TSet<UActorModifierCoreBase*>& InModifiers, AActor* InTargetActor, UActorModifierCoreBase* InTargetModifier, EActorModifierCoreStackPosition InPosition, TArray<UActorModifierCoreBase*>& OutMoveModifiers, TArray<UActorModifierCoreBase*>& OutCloneModifiers) const;
+
+	/** Move modifiers to a specific target modifier position, will perform menu action transaction if wanted */
+	bool MoveModifiers(const TArray<UActorModifierCoreBase*>& InModifiers, UActorModifierCoreBase* InTargetModifier, EActorModifierCoreStackPosition InPosition, FText* OutFailReason = nullptr, bool bInShouldTransact = true) const;
+
+	/** Clone modifiers to a specific target modifier position, will perform menu action transaction if wanted */
+	bool CloneModifiers(const TArray<UActorModifierCoreBase*>& InModifiers, UActorModifierCoreBase* InTargetModifier, EActorModifierCoreStackPosition InPosition, FText* OutFailReason = nullptr, bool bInShouldTransact = true) const;
+
+	/** Register a profiler widget based on a modifier profiler class */
+	template<typename InProfilerClass, typename = typename TEnableIf<TIsDerivedFrom<InProfilerClass, FActorModifierCoreProfiler>::Value>::Type
+		, typename InWidgetClass, typename = typename TEnableIf<TIsDerivedFrom<InWidgetClass, SActorModifierCoreEditorProfiler>::Value>::Type>
+	void RegisterProfilerWidget()
+	{
+		const FName ProfilerType = GetGeneratedTypeName<InProfilerClass>();
+		ModifierProfilerWidgets.Add(ProfilerType, [](TSharedPtr<FActorModifierCoreProfiler> InProfiler)
+		{
+			return SNew(InWidgetClass, InProfiler);
+		});
+	}
+
+	/** Unregister a modifier profiler widget */
+	template<typename InProfilerClass, typename = typename TEnableIf<TIsDerivedFrom<InProfilerClass, FActorModifierCoreProfiler>::Value>::Type>
+	void UnregisterProfilerWidget()
+	{
+		const FName ProfilerType = GetGeneratedTypeName<InProfilerClass>();
+		ModifierProfilerWidgets.Remove(ProfilerType);
+	}
+
+	/** Creates a profiler widget for a modifier profiler, returns default widget if none was registered */
+	TSharedPtr<SActorModifierCoreEditorProfiler> CreateProfilerWidget(TSharedPtr<FActorModifierCoreProfiler> InProfiler);
+	
+protected:
+	/** Modifier profiler pinned stats based on profiler type */
+	TMap<FName, TSet<FName>> ModifierProfilerStats;
+
+	/** Modifier profiler widgets */
+	TMap<FName, TFunction<TSharedRef<SActorModifierCoreEditorProfiler>(TSharedPtr<FActorModifierCoreProfiler>)>> ModifierProfilerWidgets;
+	
+	/** Runtime subsystem for modifier factories */
+	TWeakObjectPtr<UActorModifierCoreSubsystem> EngineSubsystem = nullptr;
+};
