MotionDesign : Modifiers
- Cleaning exposed public API functions
- Moving editor subsystem CRUD functions to runtime subsystem and restricting usage of CRUD functions from the stack to have one entry point

#changelist validated
#virtualized

diff --git a/Engine/Plugins/Experimental/ActorModifierCore/ActorModifierCore.uplugin b/Engine/Plugins/Experimental/ActorModifierCore/ActorModifierCore.uplugin
--- a/Engine/Plugins/Experimental/ActorModifierCore/ActorModifierCore.uplugin
+++ b/Engine/Plugins/Experimental/ActorModifierCore/ActorModifierCore.uplugin
@@ -12,7 +12,7 @@
 	"SupportURL": "",
 	"CanContainContent": true,
 	"IsBetaVersion": false,
-	"IsExperimentalVersion": false,
+	"IsExperimentalVersion": true,
 	"Installed": false,
 	"PlatformAllowList": [
 		"Win64",
diff --git a/Engine/Plugins/Experimental/ActorModifierCore/Source/ActorModifierCore/ActorModifierCore.Build.cs b/Engine/Plugins/Experimental/ActorModifierCore/Source/ActorModifierCore/ActorModifierCore.Build.cs
--- a/Engine/Plugins/Experimental/ActorModifierCore/Source/ActorModifierCore/ActorModifierCore.Build.cs
+++ b/Engine/Plugins/Experimental/ActorModifierCore/Source/ActorModifierCore/ActorModifierCore.Build.cs
@@ -27,6 +27,13 @@ public class ActorModifierCore : ModuleRules
 				"TraceLog"
 			}
 		);
-		
+
+		if (Target.Type == TargetRules.TargetType.Editor)
+		{
+			PrivateDependencyModuleNames.AddRange(new string[]
+			{
+				"UnrealEd"
+			});
+		}
 	}
 }
diff --git a/Engine/Plugins/Experimental/ActorModifierCore/Source/ActorModifierCore/Private/Modifiers/ActorModifierCoreBase.cpp b/Engine/Plugins/Experimental/ActorModifierCore/Source/ActorModifierCore/Private/Modifiers/ActorModifierCoreBase.cpp
--- a/Engine/Plugins/Experimental/ActorModifierCore/Source/ActorModifierCore/Private/Modifiers/ActorModifierCoreBase.cpp
+++ b/Engine/Plugins/Experimental/ActorModifierCore/Source/ActorModifierCore/Private/Modifiers/ActorModifierCoreBase.cpp
@@ -23,8 +23,8 @@ UActorModifierCoreSharedObject* UActorModifierCoreBase::GetShared(TSubclassOf<UA
 
 	if (const UActorModifierCoreSubsystem* Subsystem = UActorModifierCoreSubsystem::Get())
 	{
-		UWorld* World = ActorModified->GetWorld();
-		return Subsystem->GetModifierSharedObject(World, InClass, bInCreateIfNone);
+		ULevel* Level = ActorModified->GetLevel();
+		return Subsystem->GetModifierSharedObject(Level, InClass, bInCreateIfNone);
 	}
 
 	return nullptr;
@@ -415,7 +415,7 @@ bool UActorModifierCoreBase::IsModifierProfiling() const
 	return false;
 }
 
-bool UActorModifierCoreBase::ProcessFunction(TFunctionRef<bool(const UActorModifierCoreBase*)> InFunction) const
+bool UActorModifierCoreBase::ProcessFunction(TFunctionRef<bool(const UActorModifierCoreBase*)> InFunction, const FActorModifierCoreStackSearchOp& InSearchOptions) const
 {
 	return InFunction(this);
 }
diff --git a/Engine/Plugins/Experimental/ActorModifierCore/Source/ActorModifierCore/Private/Modifiers/ActorModifierCoreComponent.cpp b/Engine/Plugins/Experimental/ActorModifierCore/Source/ActorModifierCore/Private/Modifiers/ActorModifierCoreComponent.cpp
--- a/Engine/Plugins/Experimental/ActorModifierCore/Source/ActorModifierCore/Private/Modifiers/ActorModifierCoreComponent.cpp
+++ b/Engine/Plugins/Experimental/ActorModifierCore/Source/ActorModifierCore/Private/Modifiers/ActorModifierCoreComponent.cpp
@@ -12,10 +12,6 @@ UActorModifierCoreComponent::UActorModifierCoreComponent()
 	}
 }
 
-UActorModifierCoreComponent::~UActorModifierCoreComponent()
-{
-}
-
 UActorModifierCoreComponent* UActorModifierCoreComponent::CreateAndExposeComponent(AActor* InParentActor)
 {
 	if (!InParentActor)
diff --git a/Engine/Plugins/Experimental/ActorModifierCore/Source/ActorModifierCore/Private/Modifiers/ActorModifierCoreDefs.cpp b/Engine/Plugins/Experimental/ActorModifierCore/Source/ActorModifierCore/Private/Modifiers/ActorModifierCoreDefs.cpp
--- a/Engine/Plugins/Experimental/ActorModifierCore/Source/ActorModifierCore/Private/Modifiers/ActorModifierCoreDefs.cpp
+++ b/Engine/Plugins/Experimental/ActorModifierCore/Source/ActorModifierCore/Private/Modifiers/ActorModifierCoreDefs.cpp
@@ -268,6 +268,12 @@ TSharedPtr<FActorModifierCoreProfiler> FActorModifierCoreMetadata::CreateProfile
 	return nullptr;
 }
 
+const FActorModifierCoreStackSearchOp& FActorModifierCoreStackSearchOp::GetDefault()
+{
+	static const FActorModifierCoreStackSearchOp Default = {};
+	return Default;
+}
+
 FActorModifierCoreScopedLock::FActorModifierCoreScopedLock(UActorModifierCoreBase* InModifier)
 {
 	if (InModifier)
diff --git a/Engine/Plugins/Experimental/ActorModifierCore/Source/ActorModifierCore/Private/Modifiers/ActorModifierCoreSharedActor.h b/Engine/Plugins/Experimental/ActorModifierCore/Source/ActorModifierCore/Private/Modifiers/ActorModifierCoreSharedActor.h
--- a/Engine/Plugins/Experimental/ActorModifierCore/Source/ActorModifierCore/Private/Modifiers/ActorModifierCoreSharedActor.h
+++ b/Engine/Plugins/Experimental/ActorModifierCore/Source/ActorModifierCore/Private/Modifiers/ActorModifierCoreSharedActor.h
@@ -8,7 +8,7 @@
 
 /**
 	Singleton actor used to keep all shared data across modifiers instances,
-	there can only be one per world, should not be visible to user
+	there can only be one per level, should not be visible to user
 	Use the modifier subsystem to query objects, do not handle it directly
 */
 UCLASS(NotPlaceable, Hidden, NotBlueprintable, NotBlueprintType)
diff --git a/Engine/Plugins/Experimental/ActorModifierCore/Source/ActorModifierCore/Private/Modifiers/ActorModifierCoreStack.cpp b/Engine/Plugins/Experimental/ActorModifierCore/Source/ActorModifierCore/Private/Modifiers/ActorModifierCoreStack.cpp
--- a/Engine/Plugins/Experimental/ActorModifierCore/Source/ActorModifierCore/Private/Modifiers/ActorModifierCoreStack.cpp
+++ b/Engine/Plugins/Experimental/ActorModifierCore/Source/ActorModifierCore/Private/Modifiers/ActorModifierCoreStack.cpp
@@ -122,34 +122,87 @@ void UActorModifierCoreStack::OnModifierAdded(EActorModifierCoreEnableReason InR
 	}
 }
 
-void UActorModifierCoreStack::GetNamedModifiers(const FName& InName, TArray<UActorModifierCoreBase*>& OutModifiers) const
+bool UActorModifierCoreStack::ProcessFunction(TFunctionRef<bool(const UActorModifierCoreBase*)> InFunction, const FActorModifierCoreStackSearchOp& InSearchOptions) const
 {
-	for (const TObjectPtr<UActorModifierCoreBase>& Modifier : Modifiers)
+	if (!InSearchOptions.bSkipStack)
 	{
-		if (Modifier && Modifier->GetModifierName() == InName)
+		if (!InFunction(this))
 		{
-			OutModifiers.Add(Modifier);
+			return false;
 		}
 	}
-}
 
-bool UActorModifierCoreStack::ProcessFunction(TFunctionRef<bool(const UActorModifierCoreBase*)> InFunction) const
-{
-	// will only read, not write anything, used for query
 	for (const TObjectPtr<UActorModifierCoreBase>& Modifier : Modifiers)
 	{
 		if (Modifier)
 		{
-			// stop when we return false
-			if (!Modifier->ProcessFunction(InFunction))
+			if (Modifier->IsModifierStack() && !InSearchOptions.bRecurse)
 			{
-				return false;
+				if (!InSearchOptions.bSkipStack)
+				{
+					if (!InFunction(this))
+					{
+						return false;
+					}
+				}
+			}
+			else
+			{
+				if (!Modifier->ProcessFunction(InFunction, InSearchOptions))
+				{
+					return false;
+				}
 			}
 		}
 	}
+
 	return true;
 }
 
+bool UActorModifierCoreStack::ProcessSearchFunction(TFunctionRef<bool(const UActorModifierCoreBase*)> InFunction, const FActorModifierCoreStackSearchOp& InSearchOptions) const
+{
+	bool bResult = true;
+
+	if (InSearchOptions.Position == EActorModifierCoreStackPosition::After)
+	{
+		bool bAfter = !InSearchOptions.PositionContext;
+		bResult = ProcessFunction([&InSearchOptions, &bAfter, &InFunction](const UActorModifierCoreBase* InModifier)->bool
+		{
+			if (InModifier == InSearchOptions.PositionContext)
+			{
+				bAfter = true;
+				return true;
+			}
+
+			if (bAfter && !InFunction(InModifier))
+			{
+				return false;
+			}
+
+			return true;
+		}, InSearchOptions);
+	}
+	else if (InSearchOptions.Position == EActorModifierCoreStackPosition::Before)
+	{
+		bResult = ProcessFunction([&InSearchOptions, &InFunction](const UActorModifierCoreBase* InModifier)->bool
+		{
+			if (InModifier == InSearchOptions.PositionContext)
+			{
+				return false;
+			}
+
+			if (!InFunction(InModifier))
+			{
+				return false;
+			}
+
+			return true;
+		}, InSearchOptions);
+	}
+
+	return bResult;
+}
+
 UActorModifierCoreBase* UActorModifierCoreStack::GetFirstModifier() const
 {
 	if (Modifiers.Num() >= 1)
@@ -168,45 +221,127 @@ UActorModifierCoreBase* UActorModifierCoreStack::GetLastModifier() const
 	return nullptr;
 }
 
-TArray<FName> UActorModifierCoreStack::GetModifierNames() const
+bool UActorModifierCoreStack::ContainsModifier(const FName& InSearchName, const FActorModifierCoreStackSearchOp& InSearchOptions) const
+{
+	return FindModifier(InSearchName, InSearchOptions) != nullptr;
+}
+
+bool UActorModifierCoreStack::ContainsModifier(const UClass* InSearchClass, const FActorModifierCoreStackSearchOp& InSearchOptions) const
 {
-	TArray<FName> OutModifiers;
-	ProcessFunction([&OutModifiers](const UActorModifierCoreBase* InModifier)->bool
+	return FindModifier(InSearchClass, InSearchOptions) != nullptr;
+}
+
+bool UActorModifierCoreStack::ContainsModifier(const UActorModifierCoreBase* InSearchModifier, const FActorModifierCoreStackSearchOp& InSearchOptions) const
+{
+	if (!InSearchModifier)
+	{
+		return false;
+	}
+
+	bool bFound = false;
+	ProcessSearchFunction([InSearchModifier, &bFound](const UActorModifierCoreBase* InModifier)->bool
 	{
-		OutModifiers.Add(InModifier->GetModifierName());
+		// false to stop processing
+		if (InModifier && InModifier == InSearchModifier)
+		{
+			bFound = true;
+			return false;
+		}
+
 		return true;
-	});
-	return OutModifiers;
+	}, InSearchOptions);
+
+	return bFound;
 }
 
-bool UActorModifierCoreStack::ContainsModifier(const FName& InSearchName) const
+UActorModifierCoreBase* UActorModifierCoreStack::FindModifier(FName InSearchName, const FActorModifierCoreStackSearchOp& InSearchOptions) const
 {
-	const bool bResult = ProcessFunction([InSearchName](const UActorModifierCoreBase* InModifier)->bool
+	UActorModifierCoreBase* FoundModifier = nullptr;
+
+	if (InSearchName.IsNone())
+	{
+		return FoundModifier;
+	}
+
+	ProcessSearchFunction([InSearchName, &FoundModifier](const UActorModifierCoreBase* InModifier)->bool
 	{
-		// false to stop processing
 		if (InModifier && InModifier->GetModifierName() == InSearchName)
 		{
+			FoundModifier = const_cast<UActorModifierCoreBase*>(InModifier);
 			return false;
 		}
 
 		return true;
-	});
-	return !bResult;
+	}, InSearchOptions);
+
+	return FoundModifier;
 }
 
-bool UActorModifierCoreStack::ContainsModifier(const UActorModifierCoreBase* InSearchModifier) const
+UActorModifierCoreBase* UActorModifierCoreStack::FindModifier(const UClass* InSearchClass, const FActorModifierCoreStackSearchOp& InSearchOptions) const
 {
-	const bool bResult = ProcessFunction([InSearchModifier](const UActorModifierCoreBase* InModifier)->bool
+	UActorModifierCoreBase* FoundModifier = nullptr;
+
+	if (!InSearchClass || !InSearchClass->IsChildOf<UActorModifierCoreBase>())
 	{
-		// false to stop processing
-		if (InModifier && InModifier == InSearchModifier)
+		return FoundModifier;
+	}
+
+	ProcessSearchFunction([InSearchClass, &FoundModifier](const UActorModifierCoreBase* InModifier)->bool
+	{
+		if (InModifier && InModifier->IsA(InSearchClass))
 		{
+			FoundModifier = const_cast<UActorModifierCoreBase*>(InModifier);
 			return false;
 		}
 
 		return true;
-	});
-	return !bResult;
+	}, InSearchOptions);
+
+	return FoundModifier;
+}
+
+TArray<UActorModifierCoreBase*> UActorModifierCoreStack::FindModifiers(FName InSearchName, const FActorModifierCoreStackSearchOp& InSearchOptions) const
+{
+	TArray<UActorModifierCoreBase*> FoundModifiers;
+
+	if (InSearchName.IsNone())
+	{
+		return FoundModifiers;
+	}
+
+	ProcessSearchFunction([InSearchName, &FoundModifiers](const UActorModifierCoreBase* InModifier)->bool
+	{
+		if (InModifier && InModifier->GetModifierName() == InSearchName)
+		{
+			FoundModifiers.Add(const_cast<UActorModifierCoreBase*>(InModifier));
+		}
+
+		return true;
+	}, InSearchOptions);
+
+	return FoundModifiers;
+}
+
+TArray<UActorModifierCoreBase*> UActorModifierCoreStack::FindModifiers(const UClass* InSearchClass, const FActorModifierCoreStackSearchOp& InSearchOptions) const
+{
+	TArray<UActorModifierCoreBase*> FoundModifiers;
+
+	if (!InSearchClass || !InSearchClass->IsChildOf<UActorModifierCoreBase>())
+	{
+		return FoundModifiers;
+	}
+
+	ProcessSearchFunction([InSearchClass, &FoundModifiers](const UActorModifierCoreBase* InModifier)->bool
+	{
+		if (InModifier && InModifier->IsA(InSearchClass))
+		{
+			FoundModifiers.Add(const_cast<UActorModifierCoreBase*>(InModifier));
+		}
+
+		return true;
+	}, InSearchOptions);
+
+	return FoundModifiers;
 }
 
 bool UActorModifierCoreStack::ContainsModifierBefore(const FName& InSearchName, const UActorModifierCoreBase* InBeforeModifier) const
diff --git a/Engine/Plugins/Experimental/ActorModifierCore/Source/ActorModifierCore/Private/Subsystems/ActorModifierCoreSubsystem.cpp b/Engine/Plugins/Experimental/ActorModifierCore/Source/ActorModifierCore/Private/Subsystems/ActorModifierCoreSubsystem.cpp
--- a/Engine/Plugins/Experimental/ActorModifierCore/Source/ActorModifierCore/Private/Subsystems/ActorModifierCoreSubsystem.cpp
+++ b/Engine/Plugins/Experimental/ActorModifierCore/Source/ActorModifierCore/Private/Subsystems/ActorModifierCoreSubsystem.cpp
@@ -5,12 +5,17 @@
 #include "EngineUtils.h"
 #include "Engine/Engine.h"
 #include "GameFramework/Actor.h"
+#include "Internationalization/Text.h"
 #include "Modifiers/ActorModifierCoreComponent.h"
 #include "Modifiers/ActorModifierCoreSharedActor.h"
 #include "Modifiers/ActorModifierCoreStack.h"
 #include "UObject/Class.h"
 #include "UObject/UObjectIterator.h"
 
+#if WITH_EDITOR
+#include "ScopedTransaction.h"
+#endif
+
 #define LOCTEXT_NAMESPACE "ActorModifierCoreSubsystem"
 
 DEFINE_LOG_CATEGORY_STATIC(LogActorModifierCoreSubsystem, Log, All);
@@ -368,6 +373,219 @@ TArray<UActorModifierCoreBase*> UActorModifierCoreSubsystem::GetAllowedMoveModif
 	return AllowedModifiers;
 }
 
+void UActorModifierCoreSubsystem::GetSortedModifiers(const TSet<UActorModifierCoreBase*>& InModifiers, AActor* InTargetActor, UActorModifierCoreBase* InTargetModifier, EActorModifierCoreStackPosition InPosition, TArray<UActorModifierCoreBase*>& OutMoveModifiers, TArray<UActorModifierCoreBase*>& OutCloneModifiers) const
+{
+	const UActorModifierCoreSubsystem* const ModifierSubsystem = UActorModifierCoreSubsystem::Get();
+
+	if (!IsValid(ModifierSubsystem) || InModifiers.IsEmpty() || !IsValid(InTargetActor))
+	{
+		return;
+	}
+
+	for (UActorModifierCoreBase* Modifier : InModifiers)
+	{
+		if (!IsValid(Modifier))
+		{
+			continue;
+		}
+
+		// it's a clone operation if target actor is different as modifier actor
+		if (Modifier->GetModifiedActor() != InTargetActor)
+		{
+			OutCloneModifiers.Add(Modifier);
+		}
+		// it's a move operation if target actor is same as modifier actor
+		else
+		{
+			OutMoveModifiers.Add(Modifier);
+		}
+	}
+
+	// Remove unsupported modifiers when inserting where target modifier is nullptr
+	if (!InTargetModifier)
+	{
+		const TSet<FName> AllowedModifiers = ModifierSubsystem->GetAllowedModifiers(InTargetActor, InTargetModifier, InPosition);
+
+		OutMoveModifiers.RemoveAll([&AllowedModifiers](UActorModifierCoreBase* InModifier)
+		{
+			return !IsValid(InModifier) || !AllowedModifiers.Contains(InModifier->GetModifierName());
+		});
+
+		OutCloneModifiers.RemoveAll([&AllowedModifiers](UActorModifierCoreBase* InModifier)
+		{
+			return !IsValid(InModifier) || !AllowedModifiers.Contains(InModifier->GetModifierName());
+		});
+	}
+
+	// Sort them by dependency and current order
+	OutMoveModifiers.Sort([InTargetModifier](const UActorModifierCoreBase& A, const UActorModifierCoreBase& B)
+	{
+		const FActorModifierCoreMetadata& MetadataA = A.GetModifierMetadata();
+		const FActorModifierCoreMetadata& MetadataB = B.GetModifierMetadata();
+		const bool bADependsOnB = MetadataA.DependsOn(MetadataB.GetName());
+
+		const UActorModifierCoreStack* StackA = A.GetModifierStack();
+		const UActorModifierCoreStack* StackB = B.GetModifierStack();
+		const bool bAIsBeforeB = StackA == StackB && StackA->ContainsModifierBefore(&A, &B);
+
+		bool bSortOrder = bAIsBeforeB && !bADependsOnB;
+
+		const UActorModifierCoreBase* DependencyModifier = bADependsOnB ? &A : &B;
+		if (StackA->ContainsModifierAfter(InTargetModifier, DependencyModifier))
+		{
+			bSortOrder = !bSortOrder;
+		}
+
+		return bSortOrder;
+	});
+
+	OutCloneModifiers.Sort([](const UActorModifierCoreBase& A, const UActorModifierCoreBase& B)
+	{
+		const FActorModifierCoreMetadata& MetadataA = A.GetModifierMetadata();
+		const FActorModifierCoreMetadata& MetadataB = B.GetModifierMetadata();
+		const bool bADependsOnB = MetadataA.DependsOn(MetadataB.GetName());
+
+		const UActorModifierCoreStack* StackA = A.GetModifierStack();
+		const UActorModifierCoreStack* StackB = B.GetModifierStack();
+		const bool bAIsBeforeB = StackA == StackB && StackA->ContainsModifierBefore(&A, &B);
+
+		const bool bSortOrder = bAIsBeforeB && !bADependsOnB;
+
+		return bSortOrder;
+	});
+}
+
+bool UActorModifierCoreSubsystem::MoveModifiers(const TArray<UActorModifierCoreBase*>& InModifiers, UActorModifierCoreStack* InStack, FActorModifierCoreStackMoveOp& InMoveOp) const
+{
+	if (!IsValid(InStack) || InModifiers.IsEmpty())
+	{
+		return false;
+	}
+
+#if WITH_EDITOR
+	static const FText TransactionText = LOCTEXT("MoveModifiers", "Moving {0} modifier(s) {1} modifier {2}");
+	FScopedTransaction Transaction(
+		FText::Format(TransactionText
+			, FText::FromString(FString::FromInt(InModifiers.Num()))
+			, FText::FromString(InMoveOp.MovePosition == EActorModifierCoreStackPosition::After ? TEXT("after") : TEXT("before"))
+			, FText::FromName(InMoveOp.MovePositionContext ? InMoveOp.MovePositionContext->GetModifierName() : InStack->GetModifierName()))
+		, InMoveOp.bShouldTransact);
+#endif
+
+	uint32 EditModifierCount = 0;
+	UActorModifierCoreBase* TargetModifier = InMoveOp.MovePositionContext;
+
+	InStack->ProcessLockFunction([InStack, TargetModifier, &InModifiers, &InMoveOp, &EditModifierCount]()
+	{
+		UActorModifierCoreBase* OperationContext = nullptr;
+
+		for (UActorModifierCoreBase* Modifier : InModifiers)
+		{
+			if (IsValid(Modifier))
+			{
+				InMoveOp.MoveModifier = Modifier;
+				InMoveOp.MovePositionContext = TargetModifier;
+
+				// Eg 1: when moving [bend, subdivide] before target modifier X : X is after them in the stack, we want bend added before X and subdivide added before Bend
+				// Eg 2: when moving [subdivide, bend] after target modifier X : X is before them in the stack, we want subdivide after X and bend after subdivide
+				if (OperationContext &&
+					((InMoveOp.MovePosition == EActorModifierCoreStackPosition::After && InStack->ContainsModifierAfter(Modifier, TargetModifier)) ||
+					(InMoveOp.MovePosition == EActorModifierCoreStackPosition::Before && InStack->ContainsModifierBefore(Modifier, TargetModifier))))
+				{
+					InMoveOp.MovePositionContext = OperationContext;
+				}
+
+				if (InStack->MoveModifier(InMoveOp))
+				{
+					++EditModifierCount;
+					OperationContext = Modifier;
+				}
+				else
+				{
+					const AActor* TargetActor = TargetModifier->GetModifiedActor();
+					const FText& ErrorText = InMoveOp.FailReason ? *InMoveOp.FailReason : FText::GetEmpty();
+
+					// Move modifiers on actor failed
+					UE_LOG(LogActorModifierCoreSubsystem, Warning, TEXT("Move modifier %s on actor %s failed : %s"),
+						*InMoveOp.MoveModifier->GetModifierName().ToString(),
+						*TargetActor->GetActorNameOrLabel(),
+						*ErrorText.ToString());
+
+					break;
+				}
+			}
+		}
+	});
+
+	return EditModifierCount > 0;
+}
+
+TArray<UActorModifierCoreBase*> UActorModifierCoreSubsystem::CloneModifiers(const TArray<UActorModifierCoreBase*>& InModifiers, UActorModifierCoreStack* InStack, FActorModifierCoreStackCloneOp& InCloneOp) const
+{
+	TArray<UActorModifierCoreBase*> NewModifiers;
+
+	if (!IsValid(InStack) || InModifiers.IsEmpty())
+	{
+		return NewModifiers;
+	}
+
+#if WITH_EDITOR
+	static const FText TransactionText = LOCTEXT("CloneModifiers", "Cloning {0} modifier(s) {1} modifier {2}");
+	FScopedTransaction Transaction(
+		FText::Format(TransactionText
+			, FText::FromString(FString::FromInt(InModifiers.Num()))
+			, FText::FromString(InCloneOp.ClonePosition == EActorModifierCoreStackPosition::After ? TEXT("after") : TEXT("before"))
+			, FText::FromName(InCloneOp.ClonePositionContext ? InCloneOp.ClonePositionContext->GetModifierName() : InStack->GetModifierName()))
+		, InCloneOp.bShouldTransact);
+#endif
+
+	uint32 EditModifierCount = 0;
+	UActorModifierCoreBase* TargetModifier = InCloneOp.ClonePositionContext;
+
+	InStack->ProcessLockFunction([InStack, TargetModifier, &InModifiers, &InCloneOp, &EditModifierCount, &NewModifiers]()
+	{
+		UActorModifierCoreBase* OperationContext = nullptr;
+		for (UActorModifierCoreBase* Modifier : InModifiers)
+		{
+			if (IsValid(Modifier))
+			{
+				InCloneOp.ClonePositionContext = TargetModifier;
+				InCloneOp.CloneModifier = Modifier;
+
+				// When inserting A, B where B depends on A after C, clone A after C then clone B after A
+				if (OperationContext && InCloneOp.ClonePosition == EActorModifierCoreStackPosition::After)
+				{
+					InCloneOp.ClonePositionContext = OperationContext;
+				}
+
+				UActorModifierCoreBase* NewClonedModifier = InStack->CloneModifier(InCloneOp);
+				NewModifiers.Add(NewClonedModifier);
+
+				if (NewClonedModifier)
+				{
+					++EditModifierCount;
+					OperationContext = NewClonedModifier;
+				}
+				else
+				{
+					const AActor* TargetActor = TargetModifier->GetModifiedActor();
+					const FText& ErrorText = InCloneOp.FailReason ? *InCloneOp.FailReason : FText::GetEmpty();
+
+					// Clone modifiers on actor failed
+					UE_LOG(LogActorModifierCoreSubsystem, Warning, TEXT("Clone modifier %s on actor %s failed : %s"),
+						*InCloneOp.CloneModifier->GetModifierName().ToString(),
+						*TargetActor->GetActorNameOrLabel(),
+						*ErrorText.ToString());
+
+					break;
+				}
+			}
+		}
+	});
+
+	return NewModifiers;
+}
+
 bool UActorModifierCoreSubsystem::ValidateModifierCreation(const FName& InName, const UActorModifierCoreStack* InStack, FText& OutFailReason, UActorModifierCoreBase* InBeforeModifier) const
 {
 	if (!IsValid(InStack))
@@ -640,7 +858,7 @@ const AActor* UActorModifierCoreSubsystem::GetModifierStackActor(const UActorMod
 	return nullptr;
 }
 
-bool UActorModifierCoreSubsystem::ForEachModifierMetadata(TFunctionRef<bool(FActorModifierCoreMetadata&)> InProcessFunction) const
+bool UActorModifierCoreSubsystem::ForEachModifierMetadata(TFunctionRef<bool(const FActorModifierCoreMetadata&)> InProcessFunction) const
 {
 	for (const TPair<FName, TSharedRef<FActorModifierCoreMetadata>>& ModifierMetadataPair : ModifiersMetadata)
 	{
@@ -652,7 +870,7 @@ bool UActorModifierCoreSubsystem::ForEachModifierMetadata(TFunctionRef<bool(FAct
 	return true;
 }
 
-bool UActorModifierCoreSubsystem::ProcessModifierMetadata(const FName& InName, TFunctionRef<bool(FActorModifierCoreMetadata&)> InProcessFunction) const
+bool UActorModifierCoreSubsystem::ProcessModifierMetadata(const FName& InName, TFunctionRef<bool(const FActorModifierCoreMetadata&)> InProcessFunction) const
 {
 	if (TSharedRef<FActorModifierCoreMetadata> const* ModifierMetadata = ModifiersMetadata.Find(InName))
 	{
@@ -661,14 +879,14 @@ bool UActorModifierCoreSubsystem::ProcessModifierMetadata(const FName& InName, T
 	return false;
 }
 
-UActorModifierCoreSharedObject* UActorModifierCoreSubsystem::GetModifierSharedObject(UWorld* InWorld, TSubclassOf<UActorModifierCoreSharedObject> InClass, bool bInCreateIfNone) const
+UActorModifierCoreSharedObject* UActorModifierCoreSubsystem::GetModifierSharedObject(ULevel* InLevel, TSubclassOf<UActorModifierCoreSharedObject> InClass, bool bInCreateIfNone) const
 {
 	if (!InClass.Get())
 	{
 		return nullptr;
 	}
 
-	if (AActorModifierCoreSharedActor* SharedActor = GetModifierSharedProvider(InWorld))
+	if (AActorModifierCoreSharedActor* SharedActor = GetModifierSharedProvider(InLevel))
 	{
 		if (bInCreateIfNone)
 		{
@@ -683,25 +901,413 @@ UActorModifierCoreSharedObject* UActorModifierCoreSubsystem::GetModifierSharedOb
 	return nullptr;
 }
 
-AActorModifierCoreSharedActor* UActorModifierCoreSubsystem::GetModifierSharedProvider(UWorld* InWorld, bool bInSpawnIfNotFound) const
+bool UActorModifierCoreSubsystem::EnableModifiers(const TSet<UActorModifierCoreBase*>& InModifiers, bool bInEnabled, bool bInShouldTransact) const
+{
+	if (InModifiers.IsEmpty())
+	{
+		return false;
+	}
+
+#if WITH_EDITOR
+	// create transaction
+	FText TransactionText;
+	if (bInEnabled)
+	{
+		TransactionText = FText::Format(
+			LOCTEXT("EnableModifiers.Enable", "Enabling {0} modifier(s)"),
+			FText::FromString(FString::FromInt(InModifiers.Num())));
+	}
+	else
+	{
+		TransactionText = FText::Format(
+			LOCTEXT("EnableModifiers.Disable", "Disabling {0} modifier(s)"),
+			FText::FromString(FString::FromInt(InModifiers.Num())));
+	}
+	FScopedTransaction Transaction(TransactionText, bInShouldTransact);
+#endif
+
+	TSet<UActorModifierCoreBase*> ModifiersSet = InModifiers;
+
+	// lets group modifiers to batch the operation, better to update the stack only once instead of many times
+	while (!ModifiersSet.IsEmpty())
+	{
+		TSet<UActorModifierCoreBase*> StackModifiers;
+
+		const UActorModifierCoreStack* CurrentStack = nullptr;
+
+		for (UActorModifierCoreBase* Modifier : ModifiersSet)
+		{
+			if (!IsValid(Modifier))
+			{
+				continue;
+			}
+
+			UActorModifierCoreStack* ModifierStack = Modifier->GetModifierStack();
+			ModifierStack = IsValid(ModifierStack) ? ModifierStack : Modifier->GetRootModifierStack();
+
+			if (!CurrentStack)
+			{
+				CurrentStack = ModifierStack;
+			}
+
+			if (ModifierStack == CurrentStack)
+			{
+				StackModifiers.Add(Modifier);
+			}
+		}
+
+		{
+			// disable stack modifiers update
+			FActorModifierCoreScopedLock Lock(StackModifiers);
+			for (UActorModifierCoreBase* Modifier : StackModifiers)
+			{
+				Modifier->SetModifierEnabled(bInEnabled);
+			}
+		}
+
+		ModifiersSet = ModifiersSet.Difference(StackModifiers);
+	}
+
+	return true;
+}
+
+bool UActorModifierCoreSubsystem::RemoveModifiers(const TSet<UActorModifierCoreBase*>& InModifiers, FActorModifierCoreStackRemoveOp& InRemoveOp) const
+{
+	if (InModifiers.IsEmpty())
+	{
+		return false;
+	}
+
+#if WITH_EDITOR
+	// create transaction
+	FText TransactionText;
+	const UActorModifierCoreBase* SingleModifier = InModifiers.Array()[0];
+	if (InModifiers.Num() == 1 && IsValid(SingleModifier))
+	{
+		TransactionText = FText::Format(
+			LOCTEXT("RemoveSingleModifiers", "Removing {0} modifier"),
+			FText::FromName(SingleModifier->GetModifierName()));
+	}
+	else
+	{
+		TransactionText = FText::Format(
+			LOCTEXT("RemoveMultipleModifiers", "Removing {0} modifier(s)"),
+			FText::FromString(FString::FromInt(InModifiers.Num())));
+	}
+	FScopedTransaction Transaction(TransactionText, InRemoveOp.bShouldTransact);
+#endif
+
+	TSet<UActorModifierCoreBase*> ModifiersSet = InModifiers;
+
+	// lets group modifiers to batch the operation, better to update the stack only once instead of many times
+	while (!ModifiersSet.IsEmpty())
+	{
+		TSet<UActorModifierCoreBase*> StackModifiers;
+
+		UActorModifierCoreStack* CurrentStack = nullptr;
+
+		for (UActorModifierCoreBase* Modifier : ModifiersSet)
+		{
+			if (!IsValid(Modifier))
+			{
+				continue;
+			}
+
+			UActorModifierCoreStack* ModifierStack = Modifier->GetModifierStack();
+			ModifierStack = IsValid(ModifierStack) ? ModifierStack : Modifier->GetRootModifierStack();
+
+			if (!CurrentStack)
+			{
+				CurrentStack = ModifierStack;
+			}
+
+			if (ModifierStack == CurrentStack)
+			{
+				StackModifiers.Add(Modifier);
+			}
+		}
+
+		// Sort by their order in stack first to avoid dependencies errors
+		bool bSuccess = true;
+		StackModifiers.Sort([CurrentStack](UActorModifierCoreBase& A, UActorModifierCoreBase& B)
+		{
+			bool bSortResult = true;
+			CurrentStack->ProcessFunction([&A, &B, &bSortResult](const UActorModifierCoreBase* InModifier)->bool
+			{
+				if (InModifier == &A || InModifier == &B)
+				{
+					bSortResult = InModifier == &A ? false : true;
+					return false;
+				}
+				return true;
+			});
+			return bSortResult;
+		});
+
+		CurrentStack->ProcessLockFunction([this, &StackModifiers, CurrentStack, &bSuccess, &InRemoveOp]()
+		{
+			for (UActorModifierCoreBase* RemoveModifier : StackModifiers)
+			{
+				InRemoveOp.RemoveModifier = RemoveModifier;
+				if (!CurrentStack->RemoveModifier(InRemoveOp))
+				{
+					bSuccess = false;
+					break;
+				}
+			}
+		});
+
+		if (!bSuccess)
+		{
+			return false;
+		}
+
+		ModifiersSet = ModifiersSet.Difference(StackModifiers);
+	}
+
+	return true;
+}
+
+bool UActorModifierCoreSubsystem::RemoveActorsModifiers(const TSet<AActor*>& InActors, bool bInShouldTransact) const
+{
+	if (InActors.IsEmpty())
+	{
+		return false;
+	}
+
+	// Get actors stack
+	TArray<UActorModifierCoreStack*> ActorStacks;
+	for (const AActor* Actor : InActors)
+	{
+		if (!IsValid(Actor))
+		{
+			continue;
+		}
+
+		if (UActorModifierCoreStack* ActorStack = GetActorModifierStack(Actor))
+		{
+			ActorStacks.Add(ActorStack);
+		}
+	}
+
+#if WITH_EDITOR
+	// create transaction
+	const FText TransactionText = FText::Format(
+		LOCTEXT("RemoveModifier", "Remove all modifiers from {0} actor(s)"),
+		FText::FromString(FString::FromInt(InActors.Num())));
+	FScopedTransaction Transaction(TransactionText, bInShouldTransact);
+#endif
+
+	// remove modifiers from actors
+	bool bSuccess = true;
+	for (UActorModifierCoreStack* ActorStack : ActorStacks)
+	{
+		if (!ActorStack->RemoveAllModifiers())
+		{
+			bSuccess = false;
+		}
+	}
+
+	return bSuccess;
+}
+
+TArray<UActorModifierCoreBase*> UActorModifierCoreSubsystem::AddActorsModifiers(const TSet<AActor*>& InActors, FActorModifierCoreStackInsertOp& InAddOp) const
+{
+	TArray<UActorModifierCoreBase*> NewModifiers;
+
+	if (!IsRegisteredModifierClass(InAddOp.NewModifierName))
+	{
+		return NewModifiers;
+	}
+
+	// Get actors stack, create if none
+	TArray<UActorModifierCoreStack*> ActorStacks;
+	for (AActor* Actor : InActors)
+	{
+		if (!IsValid(Actor))
+		{
+			continue;
+		}
+
+		UActorModifierCoreStack* ActorStack = GetActorModifierStack(Actor);
+
+		if (!ActorStack)
+		{
+			ActorStack = AddActorModifierStack(Actor);
+		}
+
+		if (IsValid(ActorStack))
+		{
+			ActorStacks.Add(ActorStack);
+		}
+	}
+
+#if WITH_EDITOR
+	// create transaction
+	const FText TransactionText = FText::Format(
+		LOCTEXT("AddModifier", "Add {0} modifier on {1} actor(s)"),
+		FText::FromName(InAddOp.NewModifierName),
+		FText::FromString(FString::FromInt(InActors.Num())));
+	FScopedTransaction Transaction(TransactionText, InAddOp.bShouldTransact);
+#endif
+
+	// add modifier to actors
+	InAddOp.InsertPositionContext = nullptr;
+	InAddOp.InsertPosition = EActorModifierCoreStackPosition::Before;
+
+	for (UActorModifierCoreStack* ActorStack : ActorStacks)
+	{
+		NewModifiers.Add(ActorStack->InsertModifier(InAddOp));
+	}
+
+	return NewModifiers;
+}
+
+UActorModifierCoreBase* UActorModifierCoreSubsystem::InsertModifier(UActorModifierCoreStack* InStack, FActorModifierCoreStackInsertOp& InInsertOp) const
+{
+	UActorModifierCoreBase* NewModifier = nullptr;
+
+	if (!IsRegisteredModifierClass(InInsertOp.NewModifierName))
+	{
+		return NewModifier;
+	}
+
+	if (!IsValid(InStack))
+	{
+		return NewModifier;
+	}
+
+#if WITH_EDITOR
+	// create transaction
+	FText TransactionText;
+	if (InInsertOp.InsertPosition == EActorModifierCoreStackPosition::Before)
+	{
+		if (IsValid(InInsertOp.InsertPositionContext))
+		{
+			TransactionText = FText::Format(
+				LOCTEXT("InsertModifierBefore", "Insert {0} modifier before {1}"),
+				FText::FromName(InInsertOp.NewModifierName),
+				FText::FromName(InInsertOp.InsertPositionContext->GetModifierName()));
+		}
+		else
+		{
+			TransactionText = FText::Format(
+				LOCTEXT("InsertModifierAtEnd", "Insert {0} modifier at the end of stack"),
+				FText::FromName(InInsertOp.NewModifierName));
+		}
+	}
+	else if (InInsertOp.InsertPosition == EActorModifierCoreStackPosition::After)
+	{
+		if (IsValid(InInsertOp.InsertPositionContext))
+		{
+			TransactionText = FText::Format(
+				LOCTEXT("InsertModifierAfter", "Insert {0} modifier after {1}"),
+				FText::FromName(InInsertOp.NewModifierName),
+				FText::FromName(InInsertOp.InsertPositionContext->GetModifierName()));
+		}
+		else
+		{
+			TransactionText = FText::Format(
+				LOCTEXT("InsertModifierAtStart", "Insert {0} modifier at the start of stack"),
+				FText::FromName(InInsertOp.NewModifierName));
+		}
+	}
+	FScopedTransaction Transaction(TransactionText, InInsertOp.bShouldTransact);
+#endif
+
+	// insert modifier in stack
+	NewModifier = InStack->InsertModifier(InInsertOp);
+
+	return NewModifier;
+}
+
+bool UActorModifierCoreSubsystem::MoveModifier(UActorModifierCoreStack* InStack, FActorModifierCoreStackMoveOp& InMoveOp) const
+{
+	if (!IsValid(InMoveOp.MoveModifier))
+	{
+		return false;
+	}
+
+	UActorModifierCoreStack* ModifierStack = InMoveOp.MoveModifier->GetModifierStack();
+
+	if (!IsValid(ModifierStack) || ModifierStack != InStack)
+	{
+		return false;
+	}
+
+	const FName& MoveModifierName = InMoveOp.MoveModifier->GetModifierName();
+
+#if WITH_EDITOR
+	// create transaction
+	FText TransactionText;
+	if (InMoveOp.MovePosition == EActorModifierCoreStackPosition::Before)
+	{
+		if (IsValid(InMoveOp.MovePositionContext))
+		{
+			TransactionText = FText::Format(
+				LOCTEXT("MoveModifierBefore", "Move {0} modifier before {1}"),
+				FText::FromName(MoveModifierName),
+				FText::FromName(InMoveOp.MovePositionContext->GetModifierName()));
+		}
+		else
+		{
+			TransactionText = FText::Format(
+				LOCTEXT("MoveModifierToEnd", "Move {0} modifier at the end of stack"),
+				FText::FromName(MoveModifierName));
+		}
+	}
+	else if (InMoveOp.MovePosition == EActorModifierCoreStackPosition::After)
+	{
+		if (IsValid(InMoveOp.MovePositionContext))
+		{
+			TransactionText = FText::Format(
+				LOCTEXT("MoveModifierAfter", "Move {0} modifier after {1}"),
+				FText::FromName(MoveModifierName),
+				FText::FromName(InMoveOp.MovePositionContext->GetModifierName()));
+		}
+		else
+		{
+			TransactionText = FText::Format(
+				LOCTEXT("MoveModifierToStart", "Move {0} modifier at the start of stack"),
+				FText::FromName(MoveModifierName));
+		}
+	}
+	FScopedTransaction Transaction(TransactionText, InMoveOp.bShouldTransact);
+#endif
+
+	// move modifier in stack
+	return ModifierStack->MoveModifier(InMoveOp);
+}
+
+AActorModifierCoreSharedActor* UActorModifierCoreSubsystem::GetModifierSharedProvider(ULevel* InLevel, bool bInSpawnIfNotFound) const
 {
-	if (!InWorld)
+	if (!InLevel)
+	{
+		return nullptr;
+	}
+
+	UWorld* World = InLevel->GetWorld();
+
+	if (!World)
 	{
 		return nullptr;
 	}
 
 	// return cache one if already registered
-	const TWeakObjectPtr<AActorModifierCoreSharedActor>* SharedProvider = ModifierSharedProviders.Find(InWorld);
+	const TWeakObjectPtr<AActorModifierCoreSharedActor>* SharedProvider = ModifierSharedProviders.Find(InLevel);
 	if (SharedProvider && SharedProvider->IsValid())
 	{
 		return SharedProvider->Get();
 	}
 
 	// Take only first one into account and register it
-	for(AActorModifierCoreSharedActor* Actor : TActorRange<AActorModifierCoreSharedActor>(InWorld))
+	for(AActorModifierCoreSharedActor* Actor : TActorRange<AActorModifierCoreSharedActor>(World))
 	{
-		RegisterModifierSharedProvider(Actor);
-		return Actor;
+		if (Actor->GetLevel() == InLevel)
+		{
+			RegisterModifierSharedProvider(Actor);
+			return Actor;
+		}
 	}
 
 	if (!bInSpawnIfNotFound)
@@ -711,11 +1317,12 @@ AActorModifierCoreSharedActor* UActorModifierCoreSubsystem::GetModifierSharedPro
 
 	// Spawn new one and register it
 	FActorSpawnParameters SpawnParameters;
+	SpawnParameters.OverrideLevel = InLevel;
 #if WITH_EDITOR
 	SpawnParameters.bHideFromSceneOutliner = true;
 #endif
 
-	AActorModifierCoreSharedActor* const NewSharedActor = InWorld->SpawnActor<AActorModifierCoreSharedActor>(SpawnParameters);
+	AActorModifierCoreSharedActor* const NewSharedActor = World->SpawnActor<AActorModifierCoreSharedActor>(SpawnParameters);
 	RegisterModifierSharedProvider(NewSharedActor);
 
 	return NewSharedActor;
@@ -728,21 +1335,21 @@ bool UActorModifierCoreSubsystem::RegisterModifierSharedProvider(AActorModifierC
 		return false;
 	}
 
-	UWorld* World = InSharedActor->GetWorld();
-	if (!World)
+	ULevel* Level = InSharedActor->GetLevel();
+	if (!Level)
 	{
 		return false;
 	}
 
-	if (ModifierSharedProviders.Contains(World))
+	if (ModifierSharedProviders.Contains(Level))
 	{
 		return false;
 	}
 
 	UActorModifierCoreSubsystem* MutableThis = const_cast<UActorModifierCoreSubsystem*>(this);
-	MutableThis->ModifierSharedProviders.Add(World, InSharedActor);
+	MutableThis->ModifierSharedProviders.Add(Level, InSharedActor);
 
-	UE_LOG(LogActorModifierCoreSubsystem, Log, TEXT("Modifier shared provider registered for world %s"), *World->GetDebugDisplayName());
+	UE_LOG(LogActorModifierCoreSubsystem, Log, TEXT("Modifier shared provider registered for world %s and level %s"), *Level->GetWorld()->GetDebugDisplayName(), *Level->GetName());
 
 	return true;
 }
@@ -754,13 +1361,13 @@ bool UActorModifierCoreSubsystem::UnregisterModifierSharedProvider(const AActor*
 		return false;
 	}
 
-	const UWorld* World = InSharedActor->GetWorld();
-	if (!World)
+	const ULevel* Level = InSharedActor->GetLevel();
+	if (!Level)
 	{
 		return false;
 	}
 
-	const TWeakObjectPtr<AActorModifierCoreSharedActor>* SharedActor = ModifierSharedProviders.Find(World);
+	const TWeakObjectPtr<AActorModifierCoreSharedActor>* SharedActor = ModifierSharedProviders.Find(Level);
 	if (!SharedActor)
 	{
 		return false;
@@ -771,10 +1378,10 @@ bool UActorModifierCoreSubsystem::UnregisterModifierSharedProvider(const AActor*
 		return false;
 	}
 
-	UE_LOG(LogActorModifierCoreSubsystem, Log, TEXT("Modifier shared provider unregistered for world %s"), *World->GetDebugDisplayName());
+	UE_LOG(LogActorModifierCoreSubsystem, Log, TEXT("Modifier shared provider unregistered for world %s and level %s"), *Level->GetWorld()->GetDebugDisplayName(), *Level->GetName());
 
 	UActorModifierCoreSubsystem* MutableThis = const_cast<UActorModifierCoreSubsystem*>(this);
-	return MutableThis->ModifierSharedProviders.Remove(World) > 0;
+	return MutableThis->ModifierSharedProviders.Remove(Level) > 0;
 }
 
 void UActorModifierCoreSubsystem::ScanForModifiers()
diff --git a/Engine/Plugins/Experimental/ActorModifierCore/Source/ActorModifierCore/Public/Modifiers/ActorModifierCoreBase.h b/Engine/Plugins/Experimental/ActorModifierCore/Source/ActorModifierCore/Public/Modifiers/ActorModifierCoreBase.h
--- a/Engine/Plugins/Experimental/ActorModifierCore/Source/ActorModifierCore/Public/Modifiers/ActorModifierCoreBase.h
+++ b/Engine/Plugins/Experimental/ActorModifierCore/Source/ActorModifierCore/Public/Modifiers/ActorModifierCoreBase.h
@@ -21,8 +21,8 @@ class UActorModifierCoreSharedObject;
 class UActorModifierCoreStack;
 
 /** Abstract base class for all modifier, a modifier must be located in a modifier stack to work properly */
-UCLASS(NotBlueprintable, Abstract, EditInlineNew, HideCategories=(Tags, AssetUserData, Activation, Collision, Cooking))
-class ACTORMODIFIERCORE_API UActorModifierCoreBase : public UObject
+UCLASS(MinimalAPI, NotBlueprintable, Abstract, EditInlineNew, HideCategories=(Tags, AssetUserData, Activation, Collision, Cooking))
+class UActorModifierCoreBase : public UObject
 {
 	friend class UActorModifierCoreComponent;
 	friend class UActorModifierCoreStack;
@@ -41,10 +41,10 @@ public:
 	const FActorModifierCoreMetadata& GetModifierMetadata() const;
 
 	/** Returns this modifier unique name */
-	FName GetModifierName() const;
+	ACTORMODIFIERCORE_API FName GetModifierName() const;
 
 	/** Returns this modifier category */
-	FName GetModifierCategory() const;
+	ACTORMODIFIERCORE_API FName GetModifierCategory() const;
 
 	/** Gets the modifier profiler */
 	TSharedPtr<FActorModifierCoreProfiler> GetProfiler() const
@@ -59,10 +59,10 @@ public:
 	}
 
 	/** Whether this modifier is a stack to avoid casting */
-	bool IsModifierStack() const;
+	ACTORMODIFIERCORE_API bool IsModifierStack() const;
 
 	/** Returns the modified actor for this modifier */
-	AActor* GetModifiedActor() const;
+	ACTORMODIFIERCORE_API AActor* GetModifiedActor() const;
 
 	/** Returns the stack this modifier is in */
 	UActorModifierCoreStack* GetModifierStack() const
@@ -71,7 +71,7 @@ public:
 	}
 
 	/** Returns the top root stack this modifier is in */
-	UActorModifierCoreStack* GetRootModifierStack() const;
+	ACTORMODIFIERCORE_API UActorModifierCoreStack* GetRootModifierStack() const;
 
 	/** Returns the previous modifier found before this one, will recurse upwards, could return a stack */
 	const UActorModifierCoreBase* GetPreviousModifier() const;
@@ -86,7 +86,7 @@ public:
 	const UActorModifierCoreBase* GetNextNameModifier(const FName& InModifierName) const;
 
 	/** Set this modifier as dirty to be able to execute only dirty modifiers */
-	void MarkModifierDirty(bool bExecute = true);
+	ACTORMODIFIERCORE_API void MarkModifierDirty(bool bExecute = true);
 
 	/** Whether this modifier is set as dirty */
 	bool IsModifierDirty() const
@@ -100,16 +100,6 @@ public:
 		return bModifierIdle;
 	}
 
-	/** Whether this modifier contains other modifiers */
-	virtual bool IsModifierEmpty() const { return false; }
-
-	UFUNCTION()
-	void SetModifierEnabled(bool bInEnabled);
-	bool IsModifierEnabled() const;
-
-	/** Locks this modifier execution and process a function before unlocking it */
-	void ProcessLockFunction(TFunctionRef<void()> InFunction);
-
 	/** Is the execution of this modifier locked */
 	bool IsModifierExecutionLocked() const
 	{
@@ -128,11 +118,22 @@ public:
 		return bModifierInitialized;
 	}
 
-	/** Is this modifier in profiling mode */
-	bool IsModifierProfiling() const;
+	/** Whether this modifier contains other modifiers */
+	virtual bool IsModifierEmpty() const
+	{
+		return false;
+	}
 
-	/** Execute a const function on this modifier, only to read data */
-	virtual bool ProcessFunction(TFunctionRef<bool(const UActorModifierCoreBase*)> InFunction) const;
+	UFUNCTION()
+	void SetModifierEnabled(bool bInEnabled);
+
+	ACTORMODIFIERCORE_API bool IsModifierEnabled() const;
+
+	/** Locks this modifier execution and process a function before unlocking it */
+	ACTORMODIFIERCORE_API void ProcessLockFunction(TFunctionRef<void()> InFunction);
+
+	/** Is this modifier in profiling mode */
+	ACTORMODIFIERCORE_API bool IsModifierProfiling() const;
 
 	/** Helper to lookup components type, by default will loop through this actor components, stops when false is returned */
 	template<typename InComponentClass = UActorComponent, typename = typename TEnableIf<TIsDerivedFrom<InComponentClass, UActorComponent>::Value>::Type>
@@ -220,16 +221,16 @@ public:
 	}
 
 protected:
-	UActorModifierCoreBase();
+	ACTORMODIFIERCORE_API UActorModifierCoreBase();
 
 	//~ Begin UObject
-	virtual void PostLoad() override;
-	virtual void PostEditImport() override;
-	virtual void PostDuplicate(EDuplicateMode::Type DuplicateMode) override;
+	ACTORMODIFIERCORE_API virtual void PostLoad() override;
+	ACTORMODIFIERCORE_API virtual void PostEditImport() override;
+	ACTORMODIFIERCORE_API virtual void PostDuplicate(EDuplicateMode::Type DuplicateMode) override;
 #if WITH_EDITOR
-	virtual void PreEditUndo() override;
-	virtual void PostEditUndo() override;
-	virtual void PostEditChangeProperty(FPropertyChangedEvent& InPropertyChangedEvent) override;
+	ACTORMODIFIERCORE_API virtual void PreEditUndo() override;
+	ACTORMODIFIERCORE_API virtual void PostEditUndo() override;
+	ACTORMODIFIERCORE_API virtual void PostEditChangeProperty(FPropertyChangedEvent& InPropertyChangedEvent) override;
 #endif
 	//~ End UObject
 
@@ -249,7 +250,7 @@ protected:
 	virtual void RestorePreState() {}
 
 	/** Override in child classes, apply this modifier on the actual actor, call Next or Fail to complete the modifier execution */
-	virtual void Apply();
+	ACTORMODIFIERCORE_API virtual void Apply();
 
 	/** Called after the modifier was added to a stack and initialized after creation or serialization or duplication */
 	virtual void OnModifierAdded(EActorModifierCoreEnableReason InReason) {}
@@ -266,14 +267,11 @@ protected:
 	/** Called when modified actor transform is updated if this modifier is enabled */
 	virtual void OnModifiedActorTransformed() {}
 
-	/** Called when modifier becomes dirty */
-	virtual void OnModifierDirty(UActorModifierCoreBase* DirtyModifier, bool bExecute);
-
 	/** The modifier was executed properly, skip to the next one */
-	void Next();
+	ACTORMODIFIERCORE_API void Next();
 
 	/** The modifier failed executing, do not continue and stop here */
-	void Fail(const FText& InFailReason);
+	ACTORMODIFIERCORE_API void Fail(const FText& InFailReason);
 
 	/** Get shared object of that specific class or create one if none was found, casted version */
 	template<typename InSharedClass = UActorModifierCoreSharedObject, typename = typename TEnableIf<TIsDerivedFrom<InSharedClass, UActorModifierCoreSharedObject>::Value>::Type>
@@ -284,7 +282,7 @@ protected:
 	}
 
 	/** Get shared object of that specific class or create one if none was found */
-	UActorModifierCoreSharedObject* GetShared(TSubclassOf<UActorModifierCoreSharedObject> InClass, bool bInCreateIfNone = false) const;
+	ACTORMODIFIERCORE_API UActorModifierCoreSharedObject* GetShared(TSubclassOf<UActorModifierCoreSharedObject> InClass, bool bInCreateIfNone = false) const;
 
 	/** Adds an extension of a specific type to this modifier with its arguments or gets it if it already exists */
 	template<typename InExtensionClass, typename = typename TEnableIf<TIsDerivedFrom<InExtensionClass, FActorModifierCoreExtension>::IsDerived>::Type, typename... InArgTypes>
@@ -346,15 +344,21 @@ protected:
 	}
 
 	/** Gets an extension of this specific type */
-	FActorModifierCoreExtension* GetExtension(const FName& InExtensionType) const;
+	ACTORMODIFIERCORE_API FActorModifierCoreExtension* GetExtension(const FName& InExtensionType) const;
 
 	/** Removes an extension of this type from this modifier */
-	bool RemoveExtension(const FName& InExtensionType);
+	ACTORMODIFIERCORE_API bool RemoveExtension(const FName& InExtensionType);
 
 	/** Logs modifier message if in profiling mode or forced */
 	void LogModifier(const FString& InLog, bool bInForce = false) const;
 
 private:
+	/** Called when modifier becomes dirty */
+	ACTORMODIFIERCORE_API virtual void OnModifierDirty(UActorModifierCoreBase* DirtyModifier, bool bExecute);
+
+	/** Execute a const function on this modifier, only to read data */
+	ACTORMODIFIERCORE_API virtual bool ProcessFunction(TFunctionRef<bool(const UActorModifierCoreBase*)> InFunction, const FActorModifierCoreStackSearchOp& InSearchOptions) const;
+
 	/** INTERNAL USE ONLY, allows tickable modifier to mark themselves dirty */
 	void TickModifier(float InDeltaTime);
 
@@ -389,7 +393,7 @@ private:
 	void UnlockModifierExecution();
 
 	/** INTERNAL USE ONLY, add and setup a newly created extension for this modifier */
-	void AddExtensionInternal(const FName& InExtensionType, TSharedPtr<FActorModifierCoreExtension> InExtension);
+	ACTORMODIFIERCORE_API void AddExtensionInternal(const FName& InExtensionType, TSharedPtr<FActorModifierCoreExtension> InExtension);
 
 	/** Execute this modifier and calls the appropriate functions (SavePreState->Apply->SavePostState) when needed */
 	TFuture<bool> ExecuteModifier();
diff --git a/Engine/Plugins/Experimental/ActorModifierCore/Source/ActorModifierCore/Public/Modifiers/ActorModifierCoreComponent.h b/Engine/Plugins/Experimental/ActorModifierCore/Source/ActorModifierCore/Public/Modifiers/ActorModifierCoreComponent.h
--- a/Engine/Plugins/Experimental/ActorModifierCore/Source/ActorModifierCore/Public/Modifiers/ActorModifierCoreComponent.h
+++ b/Engine/Plugins/Experimental/ActorModifierCore/Source/ActorModifierCore/Public/Modifiers/ActorModifierCoreComponent.h
@@ -8,20 +8,23 @@
 class UActorModifierCoreStack;
 
 /** Component that contains a modifier stack and can be added on any actor */
-UCLASS(ClassGroup="Modifiers", BlueprintType, EditInlineNew, DefaultToInstanced, meta=(BlueprintSpawnableComponent), HideCategories=(Tags, Activation, Cooking, AssetUserData, Collision))
-class ACTORMODIFIERCORE_API UActorModifierCoreComponent : public UActorComponent
+UCLASS(MinimalAPI, ClassGroup="Modifiers", BlueprintType, EditInlineNew, DefaultToInstanced, meta=(BlueprintSpawnableComponent), HideCategories=(Tags, Activation, Cooking, AssetUserData, Collision))
+class UActorModifierCoreComponent : public UActorComponent
 {
 	GENERATED_BODY()
 
 public:
-	UActorModifierCoreComponent();
-	virtual ~UActorModifierCoreComponent() override;
-
-	UActorModifierCoreStack* GetModifierStack() const { return ModifierStack; }
-
 	/** Create this component for a specific actor and exposes it to the editor and for serialization */
 	static UActorModifierCoreComponent* CreateAndExposeComponent(AActor* InParentActor);
 
+	UActorModifierCoreComponent();
+
+	UActorModifierCoreStack* GetModifierStack() const
+	{
+		return ModifierStack;
+	}
+
+protected:
 	//~ Begin UActorComponent
 	virtual void OnComponentCreated() override;
 	virtual void OnComponentDestroyed(bool bDestroyingHierarchy) override;
@@ -33,7 +36,6 @@ public:
 #endif
 	//~ End UObject
 
-protected:
 	/** Initialize the inner root stack for the actor it is attached */
 	void InitializeStack();
 
diff --git a/Engine/Plugins/Experimental/ActorModifierCore/Source/ActorModifierCore/Public/Modifiers/ActorModifierCoreDefs.h b/Engine/Plugins/Experimental/ActorModifierCore/Source/ActorModifierCore/Public/Modifiers/ActorModifierCoreDefs.h
--- a/Engine/Plugins/Experimental/ActorModifierCore/Source/ActorModifierCore/Public/Modifiers/ActorModifierCoreDefs.h
+++ b/Engine/Plugins/Experimental/ActorModifierCore/Source/ActorModifierCore/Public/Modifiers/ActorModifierCoreDefs.h
@@ -84,7 +84,7 @@ enum class EActorModifierCoreStatus : uint8
 };
 
 /** Metadata for each modifier CDO, modifier instance will share same metadata as CDO */
-struct ACTORMODIFIERCORE_API FActorModifierCoreMetadata
+struct FActorModifierCoreMetadata
 {
 	static inline const FName DefaultCategory = TEXT("Default");
 #if WITH_EDITOR
@@ -110,34 +110,34 @@ struct ACTORMODIFIERCORE_API FActorModifierCoreMetadata
 	}
 
 	/** Set name of this modifier, should only be set once */
-	FActorModifierCoreMetadata& SetName(FName InName);
+	ACTORMODIFIERCORE_API FActorModifierCoreMetadata& SetName(FName InName);
 
 	/** Set Category group of this modifier, should only be set once */
-	FActorModifierCoreMetadata& SetCategory(FName InCategory);
+	ACTORMODIFIERCORE_API FActorModifierCoreMetadata& SetCategory(FName InCategory);
 
 	/** Allow modifier to tick to update when IsModifierDirtyable returns true */
-	FActorModifierCoreMetadata& AllowTick(bool bInAllowed);
+	ACTORMODIFIERCORE_API FActorModifierCoreMetadata& AllowTick(bool bInAllowed);
 
 	/** Allows multiple modifiers of the same type in the same stack */
-	FActorModifierCoreMetadata& AllowMultiple(bool bInAllowed);
+	ACTORMODIFIERCORE_API FActorModifierCoreMetadata& AllowMultiple(bool bInAllowed);
 
 	/** Add a modifier dependency for this modifier, will be added when this modifier is added */
-	FActorModifierCoreMetadata& AddDependency(const FName& InModifierName);
+	ACTORMODIFIERCORE_API FActorModifierCoreMetadata& AddDependency(const FName& InModifierName);
 
 	/** Disallows this modifier before another modifier */
-	FActorModifierCoreMetadata& DisallowBefore(const FName& InModifierName);
+	ACTORMODIFIERCORE_API FActorModifierCoreMetadata& DisallowBefore(const FName& InModifierName);
 
 	/** Disallows this modifier after another modifier */
-	FActorModifierCoreMetadata& DisallowAfter(const FName& InModifierName);
+	ACTORMODIFIERCORE_API FActorModifierCoreMetadata& DisallowAfter(const FName& InModifierName);
 
 	/** Avoid usage of this modifier before another modifier category */
-	FActorModifierCoreMetadata& AvoidBeforeCategory(const FName& InCategory);
+	ACTORMODIFIERCORE_API FActorModifierCoreMetadata& AvoidBeforeCategory(const FName& InCategory);
 
 	/** Avoid usage of this modifier after another modifier category */
-	FActorModifierCoreMetadata& AvoidAfterCategory(const FName& InCategory);
+	ACTORMODIFIERCORE_API FActorModifierCoreMetadata& AvoidAfterCategory(const FName& InCategory);
 
 	/** Sets the usage rule for this modifier, if it passes it will be available for this actor */
-	FActorModifierCoreMetadata& SetCompatibilityRule(const TFunction<bool(const AActor*)>& InModifierRule);
+	ACTORMODIFIERCORE_API FActorModifierCoreMetadata& SetCompatibilityRule(const TFunction<bool(const AActor*)>& InModifierRule);
 
 	/** Create the modifier instance */
 	UActorModifierCoreBase* CreateModifierInstance(UActorModifierCoreStack* InStack) const;
@@ -213,15 +213,15 @@ struct ACTORMODIFIERCORE_API FActorModifierCoreMetadata
 	bool ShouldAvoidAfter(FName InCategory) const;
 
 #if WITH_EDITOR
-	FActorModifierCoreMetadata& SetDisplayName(const FText& InName);
+	ACTORMODIFIERCORE_API FActorModifierCoreMetadata& SetDisplayName(const FText& InName);
 
-	FActorModifierCoreMetadata& SetDescription(const FText& InDescription);
+	ACTORMODIFIERCORE_API FActorModifierCoreMetadata& SetDescription(const FText& InDescription);
 
-	FActorModifierCoreMetadata& SetColor(const FLinearColor& InColor);
+	ACTORMODIFIERCORE_API FActorModifierCoreMetadata& SetColor(const FLinearColor& InColor);
 
-	FActorModifierCoreMetadata& SetIcon(const FSlateIcon& InIcon);
+	ACTORMODIFIERCORE_API FActorModifierCoreMetadata& SetIcon(const FSlateIcon& InIcon);
 
-	FActorModifierCoreMetadata& SetHidden(bool bInHidden);
+	ACTORMODIFIERCORE_API FActorModifierCoreMetadata& SetHidden(bool bInHidden);
 
 	bool IsHidden() const
 	{
@@ -250,7 +250,7 @@ struct ACTORMODIFIERCORE_API FActorModifierCoreMetadata
 #endif
 
 private:
-	void SetupProfilerInstanceInternal(TSharedPtr<FActorModifierCoreProfiler> InProfiler, UActorModifierCoreBase* InModifier, const FName& InProfilerType) const;
+	ACTORMODIFIERCORE_API void SetupProfilerInstanceInternal(TSharedPtr<FActorModifierCoreProfiler> InProfiler, UActorModifierCoreBase* InModifier, const FName& InProfilerType) const;
 
 	/** Unique name to use for this modifier */
 	FName Name;
@@ -323,6 +323,10 @@ struct FActorModifierCoreStackCloneOp
 	FText* FailReason = nullptr;
 
 	TArray<UActorModifierCoreBase*>* AddedDependencies = nullptr;
+
+#if WITH_EDITOR
+	bool bShouldTransact = false;
+#endif
 };
 
 /** Modifier stack insert/add operation */
@@ -337,6 +341,10 @@ struct FActorModifierCoreStackInsertOp
 	FText* FailReason = nullptr;
 
 	TArray<UActorModifierCoreBase*>* AddedDependencies = nullptr;
+
+#if WITH_EDITOR
+	bool bShouldTransact = false;
+#endif
 };
 
 /** Modifier stack move operation */
@@ -349,6 +357,10 @@ struct FActorModifierCoreStackMoveOp
 	UActorModifierCoreBase* MovePositionContext = nullptr;
 
 	FText* FailReason = nullptr;
+
+#if WITH_EDITOR
+	bool bShouldTransact = false;
+#endif
 };
 
 /** Modifier stack remove operation */
@@ -361,6 +373,28 @@ struct FActorModifierCoreStackRemoveOp
 	TArray<UActorModifierCoreBase*>* RemovedDependencies = nullptr;
 
 	FText* FailReason = nullptr;
+
+#if WITH_EDITOR
+	bool bShouldTransact = false;
+#endif
+};
+
+/** Modifier stack search operation */
+struct FActorModifierCoreStackSearchOp
+{
+	ACTORMODIFIERCORE_API static const FActorModifierCoreStackSearchOp& GetDefault();
+
+	/** Position to look relative to a context */
+	EActorModifierCoreStackPosition Position = EActorModifierCoreStackPosition::After;
+
+	/** Position context to operate search */
+	UActorModifierCoreBase* PositionContext = nullptr;
+
+	/** Skip stacks during search to only focus on modifiers */
+	bool bSkipStack = true;
+
+	/** Also looks in nested stacks */
+	bool bRecurse = true;
 };
 
 /** Will lock modifiers execution during lifetime of this struct */
diff --git a/Engine/Plugins/Experimental/ActorModifierCore/Source/ActorModifierCore/Public/Modifiers/ActorModifierCoreStack.h b/Engine/Plugins/Experimental/ActorModifierCore/Source/ActorModifierCore/Public/Modifiers/ActorModifierCoreStack.h
--- a/Engine/Plugins/Experimental/ActorModifierCore/Source/ActorModifierCore/Public/Modifiers/ActorModifierCoreStack.h
+++ b/Engine/Plugins/Experimental/ActorModifierCore/Source/ActorModifierCore/Public/Modifiers/ActorModifierCoreStack.h
@@ -10,13 +10,13 @@
 class USceneComponent;
 
 /** A modifier stack contains modifiers and is also a modifier by itself */
-UCLASS(BlueprintType, DefaultToInstanced, EditInlineNew)
-class ACTORMODIFIERCORE_API UActorModifierCoreStack : public UActorModifierCoreBase
+UCLASS(MinimalAPI, BlueprintType, DefaultToInstanced, EditInlineNew)
+class UActorModifierCoreStack : public UActorModifierCoreBase
 {
 	GENERATED_BODY()
 
 	friend class UActorModifierCoreBase;
-	friend class UActorModifierCoreComponent;
+	friend class UActorModifierCoreSubsystem;
 
 	friend class FActorModifierCoreEditorDetailCustomization;
 	friend class UActorModifierCoreEditorStackCustomization;
@@ -25,25 +25,26 @@
 	DECLARE_MULTICAST_DELEGATE_OneParam(FOnModifierUpdated, UActorModifierCoreBase* /** UpdatedItem */)
 
 	/** Called when a modifier is added to the stack */
-	static FOnModifierUpdated OnModifierAddedDelegate;
+	ACTORMODIFIERCORE_API static FOnModifierUpdated OnModifierAddedDelegate;
 
 	/** Called when a modifier is removed from the stack */
-	static FOnModifierUpdated OnModifierRemovedDelegate;
+	ACTORMODIFIERCORE_API static FOnModifierUpdated OnModifierRemovedDelegate;
 
 	/** Called when a modifier is moved in the stack */
-	static FOnModifierUpdated OnModifierMovedDelegate;
+	ACTORMODIFIERCORE_API static FOnModifierUpdated OnModifierMovedDelegate;
 
 	/** Create a new stack by passing the actor and the parent stack if there is one */
 	static UActorModifierCoreStack* Create(AActor* InActor, UActorModifierCoreStack* InParentStack = nullptr);
 
-	/** Gets all modifiers in the stack */
+	/** Gets all modifiers in this stack, does not recurse */
 	TConstArrayView<UActorModifierCoreBase*> GetModifiers() const
 	{
 		return Modifiers;
 	}
 
-	/** Get modifiers of a specific class only */
-	template <class InModifierType>
+	/** Get modifiers of a specific class only in this stack, does not recurse */
+	template <class InModifierType
+		UE_REQUIRES(std::is_base_of<UActorModifierCoreBase, InModifierType>::value)>
 	void GetClassModifiers(TArray<InModifierType*>& OutModifiers) const
 	{
 		for (const TObjectPtr<UActorModifierCoreBase>& Modifier : Modifiers)
@@ -53,23 +54,58 @@
 				OutModifiers.Add(static_cast<InModifierType*>(Modifier));
 			}
 		}
-	};
+	}
 
-	/** Returns the name of modifiers included in this stack and the ones below */
-	TArray<FName> GetModifierNames() const;
+	/** Gets the first modifier in this stack, does not recurse */
+	ACTORMODIFIERCORE_API UActorModifierCoreBase* GetFirstModifier() const;
 
-	/** Get modifiers with a specific name only */
-	void GetNamedModifiers(const FName& InName, TArray<UActorModifierCoreBase*>& OutModifiers) const;
+	/** Gets the last modifier in this stack, does not recurse */
+	ACTORMODIFIERCORE_API UActorModifierCoreBase* GetLastModifier() const;
 
-	/** Gets the first modifier in this stack, could return a stack */
-	UActorModifierCoreBase* GetFirstModifier() const;
-
-	/** Gets the last modifier in this stack, could return a stack */
-	UActorModifierCoreBase* GetLastModifier() const;
-
 	/** Check that we have a modifier inside this stack, checks also nested stacks */
-	bool ContainsModifier(const FName& InSearchName) const;
-	bool ContainsModifier(const UActorModifierCoreBase* InSearchModifier) const;
+	ACTORMODIFIERCORE_API bool ContainsModifier(const FName& InSearchName, const FActorModifierCoreStackSearchOp& InSearchOptions = FActorModifierCoreStackSearchOp::GetDefault()) const;
+	ACTORMODIFIERCORE_API bool ContainsModifier(const UClass* InSearchClass, const FActorModifierCoreStackSearchOp& InSearchOptions = FActorModifierCoreStackSearchOp::GetDefault()) const;
+	ACTORMODIFIERCORE_API bool ContainsModifier(const UActorModifierCoreBase* InSearchModifier, const FActorModifierCoreStackSearchOp& InSearchOptions = FActorModifierCoreStackSearchOp::GetDefault()) const;
 
+	/** Finds a modifier inside this stack, returns first found, checks also nested stacks */
+	ACTORMODIFIERCORE_API UActorModifierCoreBase* FindModifier(FName InSearchName, const FActorModifierCoreStackSearchOp& InSearchOptions = FActorModifierCoreStackSearchOp::GetDefault()) const;
+	ACTORMODIFIERCORE_API UActorModifierCoreBase* FindModifier(const UClass* InSearchClass, const FActorModifierCoreStackSearchOp& InSearchOptions = FActorModifierCoreStackSearchOp::GetDefault()) const;
+
+	/** Finds modifiers inside this stack, checks also nested stacks */
+	ACTORMODIFIERCORE_API TArray<UActorModifierCoreBase*> FindModifiers(FName InSearchName, const FActorModifierCoreStackSearchOp& InSearchOptions = FActorModifierCoreStackSearchOp::GetDefault()) const;
+	ACTORMODIFIERCORE_API TArray<UActorModifierCoreBase*> FindModifiers(const UClass* InSearchClass, const FActorModifierCoreStackSearchOp& InSearchOptions = FActorModifierCoreStackSearchOp::GetDefault()) const;
+
+	/** This is the root actor stack if we do not have any parent stack */
+	bool IsRootStack() const
+	{
+		return !ModifierStack.IsValid();
+	}
+
+	/** Execute those function when the stack is restored, before executing it again */
+	ACTORMODIFIERCORE_API void ProcessFunctionOnRestore(const TFunction<void()>& InFunction);
+
+	/** Execute those function when the stack is on idle, done with updates */
+	void ProcessFunctionOnIdle(const TFunction<void()>& InFunction);
+
+	/** Process a search function to use before/after position context, stops when false is returned */
+	ACTORMODIFIERCORE_API bool ProcessSearchFunction(TFunctionRef<bool(const UActorModifierCoreBase*)> InFunction, const FActorModifierCoreStackSearchOp& InSearchOptions) const;
+
+	/** Does this stack contains any modifiers */
+	virtual bool IsModifierEmpty() const override
+	{
+		return Modifiers.IsEmpty();
+	}
+
+	/** Set profiling mode for stack and modifiers inside */
+	ACTORMODIFIERCORE_API void SetModifierProfiling(bool bInProfiling);
+
+protected:
+	//~ Begin UObject
+	virtual void PostLoad() override;
+	//~ End UObject
+
+	/** Process a function through each modifier in the stack and also the stacks below, stop when we return false */
+	virtual bool ProcessFunction(TFunctionRef<bool(const UActorModifierCoreBase*)> InFunction, const FActorModifierCoreStackSearchOp& InSearchOptions = FActorModifierCoreStackSearchOp::GetDefault()) const override;
+
 	/** Checks that we have a modifier with this name inside this stack before another modifier, checks also nested stacks */
 	bool ContainsModifierBefore(const FName& InSearchName, const UActorModifierCoreBase* InBeforeModifier) const;
@@ -79,6 +115,9 @@
 	bool ContainsModifierAfter(const FName& InSearchName, const UActorModifierCoreBase* InAfterModifier) const;
 	bool ContainsModifierAfter(const UActorModifierCoreBase* InSearchModifier, const UActorModifierCoreBase* InAfterModifier) const;
 
+	/** Checks whether all modifier in this stack are initialized */
+	bool IsModifierStackInitialized() const;
+
 	/** Gets all modifiers found after this one in the stack that depends on this modifier */
 	bool GetDependentModifiers(UActorModifierCoreBase* InModifier, TSet<UActorModifierCoreBase*>& OutDependentModifiers) const;
 
@@ -100,38 +139,6 @@
 	/** Removes all modifiers from this stack in one batch to reduce updates */
 	bool RemoveAllModifiers();
 
-	/** This is the root actor stack if we do not have any parent stack */
-	bool IsRootStack() const
-	{
-		return !ModifierStack.IsValid();
-	}
-
-	/** Execute those function when the stack is restored, before executing it again */
-	void ProcessFunctionOnRestore(const TFunction<void()>& InFunction);
-
-	/** Execute those function when the stack is on idle, done with updates */
-	void ProcessFunctionOnIdle(const TFunction<void()>& InFunction);
-
-	/** Process a function through each modifier in the stack and also the stacks below, stop when we return false */
-	virtual bool ProcessFunction(TFunctionRef<bool(const UActorModifierCoreBase*)> InFunction) const override;
-
-	/** Does this stack contains any modifiers */
-	virtual bool IsModifierEmpty() const override
-	{
-		return Modifiers.IsEmpty();
-	}
-
-	/** Checks whether all modifier in this stack are initialized */
-	bool IsModifierStackInitialized() const;
-
-	/** Set profiling mode for stack and modifiers inside */
-	void SetModifierProfiling(bool bInProfiling);
-
-protected:
-	//~ Begin UObject
-	virtual void PostLoad() override;
-	//~ End UObject
-
 	/** Register this stack to the subsystem to query it later only if root stack */
 	virtual void OnModifierAdded(EActorModifierCoreEnableReason InReason) override;
 
@@ -163,20 +170,6 @@
 	UPROPERTY(BlueprintReadOnly, VisibleInstanceOnly, NoClear, Export, Instanced, Category = "Modifiers")
 	TArray<TObjectPtr<UActorModifierCoreBase>> Modifiers;
 
-	/** Contains a copy of modifiers in the stack for this round of execution, useful for restore and for query, can be different from modifiers array */
-	UPROPERTY(Transient, DuplicateTransient, NonTransactional)
-	TArray<TObjectPtr<UActorModifierCoreBase>> CurrentModifiers;
-
-	/** Contains the modifiers that we will execute this round, can be different from modifiers array */
-	UPROPERTY(Transient, DuplicateTransient, NonTransactional)
-	TArray<TObjectPtr<UActorModifierCoreBase>> ExecuteModifiers;
-
-	/** Functions to execute once when the stack is on idle */
-	TArray<TFunction<void()>> OnIdleFunctions;
-
-	/** Functions to execute once when the stack is restored */
-	TArray<TFunction<void()>> OnRestoreFunctions;
-
 private:
 	/** Sets the stack to receive tick events */
 	virtual void OnModifierCDOSetup(FActorModifierCoreMetadata& InMetadata) override;
@@ -203,6 +196,20 @@
 	/** Checks for any possible modifier optimization within the stack */
 	void CheckModifierOptimization(bool bInInvalidateAll);
 
+	/** Contains a copy of modifiers in the stack for this round of execution, useful for restore and for query, can be different from modifiers array */
+	UPROPERTY(Transient, DuplicateTransient, NonTransactional)
+	TArray<TObjectPtr<UActorModifierCoreBase>> CurrentModifiers;
+
+	/** Contains the modifiers that we will execute this round, can be different from modifiers array */
+	UPROPERTY(Transient, DuplicateTransient, NonTransactional)
+	TArray<TObjectPtr<UActorModifierCoreBase>> ExecuteModifiers;
+
+	/** Functions to execute once when the stack is on idle */
+	TArray<TFunction<void()>> OnIdleFunctions;
+
+	/** Functions to execute once when the stack is restored */
+	TArray<TFunction<void()>> OnRestoreFunctions;
+
 	/** Enable profiling for the modifiers in this stack */
 	UPROPERTY(EditInstanceOnly, Category="Modifier", meta=(DisplayName="Enable Profiling"))
 	bool bModifierProfiling = false;
diff --git a/Engine/Plugins/Experimental/ActorModifierCore/Source/ActorModifierCore/Public/Subsystems/ActorModifierCoreSubsystem.h b/Engine/Plugins/Experimental/ActorModifierCore/Source/ActorModifierCore/Public/Subsystems/ActorModifierCoreSubsystem.h
--- a/Engine/Plugins/Experimental/ActorModifierCore/Source/ActorModifierCore/Public/Subsystems/ActorModifierCoreSubsystem.h
+++ b/Engine/Plugins/Experimental/ActorModifierCore/Source/ActorModifierCore/Public/Subsystems/ActorModifierCoreSubsystem.h
@@ -15,7 +15,7 @@ class UActorModifierCoreStack;
 
 /** This subsystem handle all modifiers stack active in the engine and allows to create modifiers with registered metadata */
 UCLASS()
-class ACTORMODIFIERCORE_API UActorModifierCoreSubsystem : public UEngineSubsystem
+class UActorModifierCoreSubsystem : public UEngineSubsystem
 {
 	GENERATED_BODY()
 
@@ -28,10 +28,10 @@ public:
 	DECLARE_MULTICAST_DELEGATE_OneParam(FOnModifierClassRegistered, const FActorModifierCoreMetadata& /* ModifierMetadata */)
 
 	/** Called when a modifier class is registered */
-	static FOnModifierClassRegistered OnModifierClassRegisteredDelegate;
+	ACTORMODIFIERCORE_API static FOnModifierClassRegistered OnModifierClassRegisteredDelegate;
 
 	/** Called when a modifier class is unregistered */
-	static FOnModifierClassRegistered OnModifierClassUnregisteredDelegate;
+	ACTORMODIFIERCORE_API static FOnModifierClassRegistered OnModifierClassUnregisteredDelegate;
 
 	DECLARE_MULTICAST_DELEGATE_OneParam(FOnModifierStackRegistered, const UActorModifierCoreStack* /** ActorRootStack */)
 
@@ -43,29 +43,25 @@ public:
 
 	UActorModifierCoreSubsystem();
 
-	static UActorModifierCoreSubsystem* Get();
-
-	/** Default subsystem interface */
-	virtual void Initialize(FSubsystemCollectionBase& Collection) override;
-	virtual void Deinitialize() override;
+	ACTORMODIFIERCORE_API static UActorModifierCoreSubsystem* Get();
 
 	/** Register a modifier class or override an already existing one */
-	bool RegisterModifierClass(const UClass* InModifierClass, bool bInOverrideIfExists = false);
-	bool UnregisterModifierClass(const FName& InName);
-	bool IsRegisteredModifierClass(const FName& InName) const;
-	bool IsRegisteredModifierClass(const UClass* InClass) const;
+	ACTORMODIFIERCORE_API bool RegisterModifierClass(const UClass* InModifierClass, bool bInOverrideIfExists = false);
+	ACTORMODIFIERCORE_API bool UnregisterModifierClass(const FName& InName);
+	ACTORMODIFIERCORE_API bool IsRegisteredModifierClass(const FName& InName) const;
+	ACTORMODIFIERCORE_API bool IsRegisteredModifierClass(const UClass* InClass) const;
 
 	/** Return the modifier name or none from the modifier class provided */
-	FName GetRegisteredModifierName(const UClass* InModifierClass) const;
+	ACTORMODIFIERCORE_API FName GetRegisteredModifierName(const UClass* InModifierClass) const;
 
 	/** Return a set of class for all modifiers registered */
 	TSet<const UClass*> GetRegisteredModifierClasses() const;
 
 	/** Returns the name of modifiers that are currently registered */
-	TSet<FName> GetRegisteredModifiers() const;
+	ACTORMODIFIERCORE_API TSet<FName> GetRegisteredModifiers() const;
 
 	/** Returns the modifiers names that are allowed for specific actor/stack and before another modifier */
-	TSet<FName> GetAllowedModifiers(AActor* InActor, UActorModifierCoreBase* InContextModifier = nullptr, EActorModifierCoreStackPosition InContextPosition = EActorModifierCoreStackPosition::Before) const;
+	ACTORMODIFIERCORE_API TSet<FName> GetAllowedModifiers(AActor* InActor, UActorModifierCoreBase* InContextModifier = nullptr, EActorModifierCoreStackPosition InContextPosition = EActorModifierCoreStackPosition::Before) const;
 
 	/** Returns the modifiers name that match a specific category */
 	TSet<FName> GetCategoryModifiers(const FName& InCategory) const;
@@ -74,15 +70,66 @@ public:
 	TSet<FName> GetModifierCategories() const;
 
 	/** Returns the category this modifier is in */
-	FName GetModifierCategory(const FName& InModifier) const;
+	ACTORMODIFIERCORE_API FName GetModifierCategory(const FName& InModifier) const;
 
 #if WITH_EDITOR
 	/** Returns the modifiers hidden to the user */
-	TSet<FName> GetHiddenModifiers() const;
+	ACTORMODIFIERCORE_API TSet<FName> GetHiddenModifiers() const;
 #endif
 
 	/** Returns modifiers from the stack where we can move the provided modifier, required and dependent modifiers from this MoveModifier will not appear in the list */
-	TArray<UActorModifierCoreBase*> GetAllowedMoveModifiers(UActorModifierCoreBase* InMoveModifier) const;
+	ACTORMODIFIERCORE_API TArray<UActorModifierCoreBase*> GetAllowedMoveModifiers(UActorModifierCoreBase* InMoveModifier) const;
+
+	/** This is the correct way to retrieve a modifier stack for a specific actor if it has one */
+	ACTORMODIFIERCORE_API UActorModifierCoreStack* GetActorModifierStack(const AActor* InActor) const;
+
+	/** This is the correct way to add modifier component to a specific actor if it has none */
+	ACTORMODIFIERCORE_API UActorModifierCoreStack* AddActorModifierStack(AActor* InActor) const;
+
+	/** Return the actor this stack is linked to */
+	const AActor* GetModifierStackActor(const UActorModifierCoreStack* InStack) const;
+
+	/** Loops through each registered modifier metadata, only used to read */
+	ACTORMODIFIERCORE_API bool ForEachModifierMetadata(TFunctionRef<bool(const FActorModifierCoreMetadata&)> InProcessFunction) const;
+
+	/** Process this on matching registered metadata, only used to read */
+	ACTORMODIFIERCORE_API bool ProcessModifierMetadata(const FName& InName, TFunctionRef<bool(const FActorModifierCoreMetadata&)> InProcessFunction) const;
+
+	/** Get modifier shared object from a specific level and class, only one or none can exists per level */
+	UActorModifierCoreSharedObject* GetModifierSharedObject(ULevel* InLevel, TSubclassOf<UActorModifierCoreSharedObject> InClass, bool bInCreateIfNone = false) const;
+
+	/** Enable or disable modifiers, will update the original array and perform a transaction if wanted */
+	ACTORMODIFIERCORE_API bool EnableModifiers(const TSet<UActorModifierCoreBase*>& InModifiers, bool bInEnabled, bool bInShouldTransact = false) const;
+
+	/** Add a modifier to actors, adds a stack automatically if none is found, will update the original array and perform a transaction if wanted */
+	ACTORMODIFIERCORE_API TArray<UActorModifierCoreBase*> AddActorsModifiers(const TSet<AActor*>& InActors, FActorModifierCoreStackInsertOp& InAddOp) const;
+
+	/** Remove all modifiers from actors, will update the original array and perform a transaction if wanted */
+	ACTORMODIFIERCORE_API bool RemoveActorsModifiers(const TSet<AActor*>& InActors, bool bInShouldTransact = false) const;
+
+	/** Remove modifiers from different actors or stacks, will update the original array and perform a transaction if wanted */
+	ACTORMODIFIERCORE_API bool RemoveModifiers(const TSet<UActorModifierCoreBase*>& InModifiers, FActorModifierCoreStackRemoveOp& InRemoveOp) const;
+
+	/** Insert a modifier in a stack before or after another modifier, will perform a transaction if wanted */
+	ACTORMODIFIERCORE_API UActorModifierCoreBase* InsertModifier(UActorModifierCoreStack* InStack, FActorModifierCoreStackInsertOp& InInsertOp) const;
+
+	/** Moves a modifier in a stack before or after another modifier, will perform a transaction if wanted */
+	ACTORMODIFIERCORE_API bool MoveModifier(UActorModifierCoreStack* InStack, FActorModifierCoreStackMoveOp& InMoveOp) const;
+
+	/** Move modifiers to a specific target modifier position, will perform menu action transaction if wanted */
+	ACTORMODIFIERCORE_API bool MoveModifiers(const TArray<UActorModifierCoreBase*>& InModifiers, UActorModifierCoreStack* InStack, FActorModifierCoreStackMoveOp& InMoveOp) const;
+
+	/** Clone modifiers to a specific target modifier position, will perform menu action transaction if wanted */
+	ACTORMODIFIERCORE_API TArray<UActorModifierCoreBase*> CloneModifiers(const TArray<UActorModifierCoreBase*>& InModifiers, UActorModifierCoreStack* InStack, FActorModifierCoreStackCloneOp& InCloneOp) const;
+
+	/** Get order dependent modifiers in the correct order for stack operations */
+	ACTORMODIFIERCORE_API void GetSortedModifiers(const TSet<UActorModifierCoreBase*>& InModifiers, AActor* InTargetActor, UActorModifierCoreBase* InTargetModifier, EActorModifierCoreStackPosition InPosition, TArray<UActorModifierCoreBase*>& OutMoveModifiers, TArray<UActorModifierCoreBase*>& OutCloneModifiers) const;
+
+protected:
+	//~ Begin UEngineSubsystem
+	virtual void Initialize(FSubsystemCollectionBase& Collection) override;
+	virtual void Deinitialize() override;
+	//~ End UEngineSubsystem
 
 	/** Register a root actor stack to query this actor stack from everywhere, is called automatically at creation or deserialization by stack */
 	bool RegisterActorModifierStack(UActorModifierCoreStack* InStack);
@@ -99,25 +146,6 @@ public:
 	/** Do we have a registered root stack */
 	bool IsRegisteredActorModifierStack(const UActorModifierCoreStack* InStack) const;
 
-	/** This is the correct way to retrieve a modifier stack for a specific actor if it has one */
-	UActorModifierCoreStack* GetActorModifierStack(const AActor* InActor) const;
-
-	/** This is the correct way to add modifier component to a specific actor if it has none */
-	UActorModifierCoreStack* AddActorModifierStack(AActor* InActor) const;
-
-	/** Return the actor this stack is linked to */
-	const AActor* GetModifierStackActor(const UActorModifierCoreStack* InStack) const;
-
-	/** Loops through each registered modifier metadata, only used to read */
-	bool ForEachModifierMetadata(TFunctionRef<bool(FActorModifierCoreMetadata&)> InProcessFunction) const;
-
-	/** Process this on matching registered metadata, only used to read */
-	bool ProcessModifierMetadata(const FName& InName, TFunctionRef<bool(FActorModifierCoreMetadata&)> InProcessFunction) const;
-
-	/** Get modifier shared object from a specific world and class, only one or none can exists per world */
-	UActorModifierCoreSharedObject* GetModifierSharedObject(UWorld* InWorld, TSubclassOf<UActorModifierCoreSharedObject> InClass, bool bInCreateIfNone = false) const;
-
-protected:
 	/** Checks whether we can add a modifier to the stack */
 	bool ValidateModifierCreation(const FName& InName, const UActorModifierCoreStack* InStack, FText& OutFailReason, UActorModifierCoreBase* InBeforeModifier = nullptr) const;
 
@@ -137,7 +165,7 @@ protected:
 	bool UnregisterModifierSharedProvider(const AActor* InSharedActor) const;
 
 	/** Gets a modifier shared provider for a world, creates one if none is found */
-	AActorModifierCoreSharedActor* GetModifierSharedProvider(UWorld* InWorld, bool bInSpawnIfNotFound = true) const;
+	AActorModifierCoreSharedActor* GetModifierSharedProvider(ULevel* InLevel, bool bInSpawnIfNotFound = true) const;
 
 	/** stores modifiers metadata, you can override these if you want a different behaviour */
 	TMap<FName, TSharedRef<FActorModifierCoreMetadata>> ModifiersMetadata;
@@ -145,6 +173,6 @@ protected:
 	/** stores modifiers stacks per actor, there should be only one root stack per actor */
 	TMap<TWeakObjectPtr<const AActor>, TWeakObjectPtr<UActorModifierCoreStack>> ModifierStacks;
 
-	/** Stores modifiers providers per world, there should be only one provider per world */
-	TMap<TWeakObjectPtr<const UWorld>, TWeakObjectPtr<AActorModifierCoreSharedActor>> ModifierSharedProviders;
+	/** Stores modifiers providers per world/level, there should be only one provider per level */
+	TMap<TWeakObjectPtr<const ULevel>, TWeakObjectPtr<AActorModifierCoreSharedActor>> ModifierSharedProviders;
 };
diff --git a/Engine/Plugins/Experimental/ActorModifierCore/Source/ActorModifierCoreEditor/Private/Modifiers/ActorModifierCoreEditorMenu.cpp b/Engine/Plugins/Experimental/ActorModifierCore/Source/ActorModifierCoreEditor/Private/Modifiers/ActorModifierCoreEditorMenu.cpp
--- a/Engine/Plugins/Experimental/ActorModifierCore/Source/ActorModifierCoreEditor/Private/Modifiers/ActorModifierCoreEditorMenu.cpp
+++ b/Engine/Plugins/Experimental/ActorModifierCore/Source/ActorModifierCoreEditor/Private/Modifiers/ActorModifierCoreEditorMenu.cpp
@@ -15,7 +15,7 @@
 void UE::ActorModifierCoreEditor::OnExtendAddModifierMenu(UToolMenu* InAddToolMenu, const FActorModifierCoreEditorMenuData& InData)
 {
 	const UActorModifierCoreSubsystem* ModifierSubsystem = UActorModifierCoreSubsystem::Get();
-	
+
 	if (!ModifierSubsystem || !InAddToolMenu)
 	{
 		return;
@@ -23,7 +23,7 @@ void UE::ActorModifierCoreEditor::OnExtendAddModifierMenu(UToolMenu* InAddToolMe
 
 	TSet<FName> ShowModifiers = ModifierSubsystem->GetRegisteredModifiers();
 	ShowModifiers = ShowModifiers.Difference(ModifierSubsystem->GetHiddenModifiers());
-	for (AActor* Actor : InData.Context.ContextActors)
+	for (AActor* Actor : InData.Context.GetContextActors())
 	{
 		ShowModifiers = ShowModifiers.Intersect(ModifierSubsystem->GetAllowedModifiers(Actor));
 	}
@@ -49,7 +49,7 @@ void UE::ActorModifierCoreEditor::OnExtendAddModifierMenu(UToolMenu* InAddToolMe
 			ModifierDescription = InMetadata.GetDescription();
 			return true;
 		});
-		
+
 		const FName ModifierCategorySection("ContextModifierCategory" + ModifierCategory.ToString());
 		FToolMenuSection* CategorySection = InAddToolMenu->FindSection(ModifierCategorySection);
 		if (!CategorySection)
@@ -59,16 +59,16 @@ void UE::ActorModifierCoreEditor::OnExtendAddModifierMenu(UToolMenu* InAddToolMe
 				ModifierDisplayCategory,
 				FToolMenuInsert(NAME_None, EToolMenuInsertType::First));
 		}
-		
+
 		const FText Label = FText::Format(
 			LOCTEXT("AddModifier.Label", "{0}"),
 			ModifierDisplayName);
-			
+
 		const FText Tooltip = FText::Format(
 			LOCTEXT("AddModifier.Tooltip", "Add Modifier {0} to the currently selected actors\n{1}"),
 			ModifierDisplayName,
 			ModifierDescription);
-		
+
 		const FUIAction Action = FUIAction(FExecuteAction::CreateLambda(&UE::ActorModifierCoreEditor::OnAddModifierMenuAction, Modifier, InData));
 
 		const FName MenuEntryName(TEXT("AddModifier") + Modifier.ToString());
@@ -86,18 +86,18 @@ void UE::ActorModifierCoreEditor::OnExtendAddModifierMenu(UToolMenu* InAddToolMe
 void UE::ActorModifierCoreEditor::OnExtendInsertModifierMenu(UToolMenu* InInsertToolMenu, const FActorModifierCoreEditorMenuData& InData)
 {
 	const UActorModifierCoreSubsystem* ModifierSubsystem = UActorModifierCoreSubsystem::Get();
-	
+
 	if (!ModifierSubsystem || !InInsertToolMenu)
 	{
 		return;
 	}
-	
-	if (InData.Context.ContextModifiers.IsEmpty() || InData.Context.ContextModifiers.Num() > 1)
+
+	if (InData.Context.GetContextModifiers().IsEmpty() || InData.Context.GetContextModifiers().Num() > 1)
 	{
 		return;
 	}
-	
-	UActorModifierCoreBase* ContextModifier = InData.Context.ContextModifiers.Array()[0];
+
+	UActorModifierCoreBase* ContextModifier = InData.Context.GetContextModifiers().Array()[0];
 
 	if (!ContextModifier)
 	{
@@ -138,17 +138,17 @@ void UE::ActorModifierCoreEditor::OnExtendInsertModifierMenu(UToolMenu* InInsert
 				ModifierDisplayCategory,
 				FToolMenuInsert(NAME_None, EToolMenuInsertType::First));
 		}
-		
+
 		const FText Label = FText::Format(
-			LOCTEXT("InsertModifier.Label", "{0}"), 
+			LOCTEXT("InsertModifier.Label", "{0}"),
 			ModifierDisplayName);
-			
+
 		const FText Tooltip = FText::Format(
 			LOCTEXT("InsertModifier.Tooltip", "Insert Modifier {0} {1} current selection\n{2}"),
 			ModifierDisplayName,
 			FText::FromString(InsertPosition == EActorModifierCoreStackPosition::Before ? TEXT("before") : TEXT("after")),
 			ModifierDescription);
-		
+
 		const FUIAction Action = FUIAction(FExecuteAction::CreateLambda(&UE::ActorModifierCoreEditor::OnInsertModifierMenuAction, ContextModifier, Modifier, InData));
 
 		const FName MenuEntryName(TEXT("InsertModifier") + Modifier.ToString());
@@ -166,12 +166,12 @@ void UE::ActorModifierCoreEditor::OnExtendInsertModifierMenu(UToolMenu* InInsert
 void UE::ActorModifierCoreEditor::OnExtendRemoveModifierMenu(UToolMenu* InRemoveToolMenu, const FActorModifierCoreEditorMenuData& InData)
 {
 	const UActorModifierCoreSubsystem* ModifierSubsystem = UActorModifierCoreSubsystem::Get();
-	
+
 	if (!ModifierSubsystem || !InRemoveToolMenu)
 	{
 		return;
 	}
-	
+
 	if (InData.Context.IsEmpty())
 	{
 		return;
@@ -186,17 +186,17 @@ void UE::ActorModifierCoreEditor::OnExtendRemoveModifierMenu(UToolMenu* InRemove
 		static const FSlateIcon RemoveAllIcon = FSlateIconFinder::FindIconForClass(UActorModifierCoreBase::StaticClass());
 
 		const FUIAction Action = FUIAction(FExecuteAction::CreateLambda(&UE::ActorModifierCoreEditor::OnRemoveActorModifierMenuAction, InData));
-		
+
 		FToolMenuEntry RemoveAllModifierEntry = FToolMenuEntry::InitMenuEntry(
 			RemoveAllEntryName,
 			RemoveAllLabel,
 			RemoveAllTooltip,
 			RemoveAllIcon,
 			Action);
-	
+
 		InRemoveToolMenu->AddMenuEntry(RemoveAllEntryName, RemoveAllModifierEntry);
 	}
-	
+
 	// remove all selected modifiers
 	if (InData.Context.ContainsAnyModifier())
 	{
@@ -206,34 +206,34 @@ void UE::ActorModifierCoreEditor::OnExtendRemoveModifierMenu(UToolMenu* InRemove
 		static const FSlateIcon RemoveSelectedIcon = FSlateIconFinder::FindIconForClass(UActorModifierCoreBase::StaticClass());
 
 		const FUIAction Action = FUIAction(FExecuteAction::CreateLambda(&UE::ActorModifierCoreEditor::OnRemoveModifierMenuAction, InData));
-		
+
 		FToolMenuEntry RemoveSelectedModifierEntry = FToolMenuEntry::InitMenuEntry(
 			RemoveSelectedEntryName,
 			RemoveSelectedLabel,
 			RemoveSelectedTooltip,
 			RemoveSelectedIcon,
 			Action);
-	
+
 		InRemoveToolMenu->AddMenuEntry(RemoveSelectedEntryName, RemoveSelectedModifierEntry);
 	}
 
 	// only handle one actor selected case after this
-	if (!InData.Context.ContainsAnyActor() || InData.Context.ContextActors.Num() > 1)
+	if (!InData.Context.ContainsAnyActor() || InData.Context.GetContextActors().Num() > 1)
 	{
 		return;
 	}
 
 	// get current actor stack
-	TWeakObjectPtr<AActor> ContextActor = InData.Context.ContextActors.Array()[0];
+	TWeakObjectPtr<AActor> ContextActor = InData.Context.GetContextActors().Array()[0];
 	const UActorModifierCoreStack* ModifierStack = ModifierSubsystem->GetActorModifierStack(ContextActor.Get());
-	
+
 	if (!ModifierStack)
 	{
 		return;
 	}
 
 	TConstArrayView<UActorModifierCoreBase*> StackModifiers = ModifierStack->GetModifiers();
-	
+
 	if (StackModifiers.IsEmpty())
 	{
 		return;
@@ -243,7 +243,7 @@ void UE::ActorModifierCoreEditor::OnExtendRemoveModifierMenu(UToolMenu* InRemove
 	static const FName SeparatorName("RemoveAllSeparator");
 	FToolMenuEntry Separator = FToolMenuEntry::InitSeparator(SeparatorName);
 	InRemoveToolMenu->AddMenuEntry(SeparatorName, Separator);
-	
+
 	for (UActorModifierCoreBase* StackModifier : StackModifiers)
 	{
 		if (!IsValid(StackModifier))
@@ -264,22 +264,22 @@ void UE::ActorModifierCoreEditor::OnExtendRemoveModifierMenu(UToolMenu* InRemove
 			ModifierIcon = InMetadata.GetIcon();
 			return true;
 		});
-		
+
 		const FText Label = FText::Format(LOCTEXT("RemoveModifier.Label", "{0}"),
 				ModifierDisplayName);
 		const FText Tooltip = FText::Format(LOCTEXT("RemoveModifier.Tooltip", "Remove Modifier {0} from the currently selected actor"),
 				ModifierDisplayName);
 		const FName MenuEntryName(TEXT("RemoveModifier") + ModifierName.ToString());
-		
+
 		FUIAction Action = FUIAction(FExecuteAction::CreateLambda(&UE::ActorModifierCoreEditor::OnRemoveSingleModifierMenuAction, StackModifier, InData));
-		
+
 		FToolMenuEntry RemoveModifierEntry = FToolMenuEntry::InitMenuEntry(
 			MenuEntryName,
 			Label,
 			Tooltip,
 			ModifierIcon,
 			Action);
-	
+
 		InRemoveToolMenu->AddMenuEntry(MenuEntryName, RemoveModifierEntry);
 	}
 }
@@ -287,31 +287,31 @@ void UE::ActorModifierCoreEditor::OnExtendRemoveModifierMenu(UToolMenu* InRemove
 void UE::ActorModifierCoreEditor::OnExtendMoveModifierMenu(UToolMenu* InMoveToolMenu, const FActorModifierCoreEditorMenuData& InData)
 {
 	const UActorModifierCoreSubsystem* ModifierSubsystem = UActorModifierCoreSubsystem::Get();
-	
+
 	if (!ModifierSubsystem || !InMoveToolMenu)
 	{
 		return;
 	}
-	
-	if (InData.Context.ContextModifiers.IsEmpty() || InData.Context.ContextModifiers.Num() > 1)
+
+	if (InData.Context.GetContextModifiers().IsEmpty() || InData.Context.GetContextModifiers().Num() > 1)
 	{
 		return;
 	}
-	
-	UActorModifierCoreBase* MoveModifier = InData.Context.ContextModifiers.Array()[0];
+
+	UActorModifierCoreBase* MoveModifier = InData.Context.GetContextModifiers().Array()[0];
 
 	if (!MoveModifier)
 	{
 		return;
 	}
-	
+
 	TArray<UActorModifierCoreBase*> StackModifiers = ModifierSubsystem->GetAllowedMoveModifiers(MoveModifier);
-	
+
 	if (StackModifiers.IsEmpty())
 	{
 		return;
 	}
-	
+
 	const FName& MoveModifierName = MoveModifier->GetModifierName();
 
 	// move in the beginning
@@ -320,20 +320,20 @@ void UE::ActorModifierCoreEditor::OnExtendMoveModifierMenu(UToolMenu* InMoveTool
 		static const FText StartLabel = LOCTEXT("MoveModifierStart.Label", "At the start");
 		static const FText StartTooltip = LOCTEXT("MoveModifierStart.Tooltip", "Move Modifier at the beginning of the stack");
 		static const FName StartMenuEntryName(TEXT("MoveModifierStart"));
-		
+
 		UActorModifierCoreBase* StartPositionModifier = nullptr;
 		FUIAction StartAction = FUIAction(FExecuteAction::CreateLambda(&UE::ActorModifierCoreEditor::OnMoveModifierMenuAction, MoveModifier, StartPositionModifier, InData));
-		
+
 		FToolMenuEntry MoveModifierEntry = FToolMenuEntry::InitMenuEntry(
 			StartMenuEntryName,
 			StartLabel,
 			StartTooltip,
 			FSlateIcon(),
 			StartAction);
-	
+
 		InMoveToolMenu->AddMenuEntry(StartMenuEntryName, MoveModifierEntry);
 	}
-	
+
 	for (UActorModifierCoreBase* PositionModifier : StackModifiers)
 	{
 		if (!IsValid(PositionModifier))
@@ -359,22 +359,22 @@ void UE::ActorModifierCoreEditor::OnExtendMoveModifierMenu(UToolMenu* InMoveTool
 			ModifierIcon = InMetadata.GetIcon();
 			return true;
 		});
-		
+
 		const FText Label = FText::Format(LOCTEXT("MoveModifier.Label", "After {0}"),
 				ModifierDisplayName);
 		const FText Tooltip = FText::Format(LOCTEXT("MoveModifier.Tooltip", "Move Modifier {0} in the stack"),
 				FText::FromName(MoveModifierName));
 		const FName MenuEntryName(TEXT("MoveModifier") + ModifierName.ToString());
-		
+
 		FUIAction Action = FUIAction(FExecuteAction::CreateLambda(&UE::ActorModifierCoreEditor::OnMoveModifierMenuAction, MoveModifier, PositionModifier, InData));
-		
+
 		FToolMenuEntry MoveModifierEntry = FToolMenuEntry::InitMenuEntry(
 			MenuEntryName,
 			Label,
 			Tooltip,
 			ModifierIcon,
 			Action);
-	
+
 		InMoveToolMenu->AddMenuEntry(MenuEntryName, MoveModifierEntry);
 	}
 }
@@ -390,34 +390,34 @@ void UE::ActorModifierCoreEditor::OnExtendEnableModifierMenu(UToolMenu* InToolMe
 			static const FText EnableStackLabel = LOCTEXT("EnableStack.Label", "Enable selected stack");
 			static const FText EnableStackTooltip = LOCTEXT("EnableStack.Tooltip", "Enable selected modifier stack");
 			static const FName EnableStackEntryName(TEXT("EnableStack"));
-		
+
 			const FUIAction EnableStackAction = FUIAction(FExecuteAction::CreateLambda(&UE::ActorModifierCoreEditor::OnEnableModifierMenuAction, InData, true, true));
-		
+
 			const FToolMenuEntry EnableStackEntry = FToolMenuEntry::InitMenuEntry(
 				EnableStackEntryName,
 				EnableStackLabel,
 				EnableStackTooltip,
 				FSlateIcon(),
 				EnableStackAction);
-	
+
 			InToolMenu->AddMenuEntry(EnableStackEntryName, EnableStackEntry);
 		}
-		
+
 		if (InData.Context.ContainsDisabledModifier())
 		{
 			static const FText EnableModifierLabel = LOCTEXT("EnableModifier.Label", "Enable selected modifier");
 			static const FText EnableModifierTooltip = LOCTEXT("EnableModifier.Tooltip", "Enable selected modifier");
 			static const FName EnableModifierEntryName(TEXT("EnableModifier"));
-		
+
 			const FUIAction EnableModifierAction = FUIAction(FExecuteAction::CreateLambda(&UE::ActorModifierCoreEditor::OnEnableModifierMenuAction, InData, true, false));
-		
+
 			const FToolMenuEntry EnableModifierEntry = FToolMenuEntry::InitMenuEntry(
 				EnableModifierEntryName,
 				EnableModifierLabel,
 				EnableModifierTooltip,
 				FSlateIcon(),
 				EnableModifierAction);
-	
+
 			InToolMenu->AddMenuEntry(EnableModifierEntryName, EnableModifierEntry);
 		}
 	}
@@ -428,34 +428,34 @@ void UE::ActorModifierCoreEditor::OnExtendEnableModifierMenu(UToolMenu* InToolMe
 			static const FText DisableStackLabel = LOCTEXT("DisableStack.Label", "Disable selected stack");
 			static const FText DisableStackTooltip = LOCTEXT("DisableStack.Tooltip", "Disable selected modifier stack");
 			static const FName DisableStackEntryName(TEXT("DisableStack"));
-			
+
 			const FUIAction DisableStackAction = FUIAction(FExecuteAction::CreateLambda(&UE::ActorModifierCoreEditor::OnEnableModifierMenuAction, InData, false, true));
-		
+
 			const FToolMenuEntry DisableStackEntry = FToolMenuEntry::InitMenuEntry(
 				DisableStackEntryName,
 				DisableStackLabel,
 				DisableStackTooltip,
 				FSlateIcon(),
 				DisableStackAction);
-	
+
 			InToolMenu->AddMenuEntry(DisableStackEntryName, DisableStackEntry);
 		}
-		
+
 		if (InData.Context.ContainsEnabledModifier())
 		{
 			static const FText DisableModifierLabel = LOCTEXT("DisableModifier.Label", "Disable selected modifier");
 			static const FText DisableModifierTooltip = LOCTEXT("DisableModifier.Tooltip", "Disable selected modifier");
 			static const FName DisableModifierEntryName(TEXT("DisableModifier"));
-		
+
 			const FUIAction DisableModifierAction = FUIAction(FExecuteAction::CreateLambda(&UE::ActorModifierCoreEditor::OnEnableModifierMenuAction, InData, false, false));
-		
+
 			const FToolMenuEntry DisableModifierEntry = FToolMenuEntry::InitMenuEntry(
 				DisableModifierEntryName,
 				DisableModifierLabel,
 				DisableModifierTooltip,
 				FSlateIcon(),
 				DisableModifierAction);
-	
+
 			InToolMenu->AddMenuEntry(DisableModifierEntryName, DisableModifierEntry);
 		}
 	}
@@ -463,21 +463,28 @@ void UE::ActorModifierCoreEditor::OnExtendEnableModifierMenu(UToolMenu* InToolMe
 
 void UE::ActorModifierCoreEditor::OnAddModifierMenuAction(const FName& InModifier, const FActorModifierCoreEditorMenuData& InData)
 {
-	const UActorModifierCoreEditorSubsystem* ModifierEditorSubsystem = UActorModifierCoreEditorSubsystem::Get();
+	const UActorModifierCoreSubsystem* ModifierSubsystem = UActorModifierCoreSubsystem::Get();
 
-	if (!IsValid(ModifierEditorSubsystem))
+	if (!IsValid(ModifierSubsystem))
 	{
 		return;
 	}
 
-	TSet<AActor*> Actors = InData.Context.ContextActors;
+	TSet<AActor*> Actors = InData.Context.GetContextActors();
 	if (Actors.IsEmpty())
 	{
 		return;
 	}
-	
+
 	FText OutFailReason;
-	if (!ModifierEditorSubsystem->AddActorsModifiers(InModifier, Actors, &OutFailReason))
+	FActorModifierCoreStackInsertOp InsertOp;
+	InsertOp.bShouldTransact = true;
+	InsertOp.FailReason = &OutFailReason;
+	InsertOp.NewModifierName = InModifier;
+
+	ModifierSubsystem->AddActorsModifiers(Actors, InsertOp);
+
+	if (!OutFailReason.IsEmpty())
 	{
 		if (InData.Options.ShouldFireNotification())
 		{
@@ -491,18 +498,25 @@ void UE::ActorModifierCoreEditor::OnAddModifierMenuAction(const FName& InModifie
 
 void UE::ActorModifierCoreEditor::OnInsertModifierMenuAction(UActorModifierCoreBase* InModifier, const FName& InModifierName, const FActorModifierCoreEditorMenuData& InData)
 {
-	const UActorModifierCoreEditorSubsystem* ModifierEditorSubsystem = UActorModifierCoreEditorSubsystem::Get();
+	const UActorModifierCoreSubsystem* ModifierSubsystem = UActorModifierCoreSubsystem::Get();
 
-	if (!IsValid(ModifierEditorSubsystem))
+	if (!IsValid(ModifierSubsystem))
 	{
 		return;
 	}
 
 	UActorModifierCoreStack* ModifierStack = InModifier->GetModifierStack();
-	
+
 	FText OutFailReason;
-	const EActorModifierCoreStackPosition InsertPosition = InData.Options.GetMenuType() == EActorModifierCoreEditorMenuType::InsertBefore ? EActorModifierCoreStackPosition::Before : EActorModifierCoreStackPosition::After;
-	if (!ModifierEditorSubsystem->InsertModifier(InModifierName, ModifierStack, InModifier, InsertPosition, &OutFailReason))
+
+	FActorModifierCoreStackInsertOp InsertOp;
+	InsertOp.bShouldTransact = true;
+	InsertOp.FailReason = &OutFailReason;
+	InsertOp.InsertPosition = InData.Options.GetMenuType() == EActorModifierCoreEditorMenuType::InsertBefore ? EActorModifierCoreStackPosition::Before : EActorModifierCoreStackPosition::After;
+	InsertOp.InsertPositionContext = InModifier;
+	InsertOp.NewModifierName = InModifierName;
+
+	if (!ModifierSubsystem->InsertModifier(ModifierStack, InsertOp))
 	{
 		if (InData.Options.ShouldFireNotification())
 		{
@@ -516,20 +530,20 @@ void UE::ActorModifierCoreEditor::OnInsertModifierMenuAction(UActorModifierCoreB
 
 void UE::ActorModifierCoreEditor::OnRemoveActorModifierMenuAction(const FActorModifierCoreEditorMenuData& InData)
 {
-	const UActorModifierCoreEditorSubsystem* ModifierEditorSubsystem = UActorModifierCoreEditorSubsystem::Get();
+	const UActorModifierCoreSubsystem* ModifierSubsystem = UActorModifierCoreSubsystem::Get();
 
-	if (!IsValid(ModifierEditorSubsystem))
+	if (!IsValid(ModifierSubsystem))
 	{
 		return;
 	}
-	
-	TSet<AActor*> Actors = InData.Context.ContextActors;
+
+	const TSet<AActor*> Actors = InData.Context.GetContextActors();
 	if (Actors.IsEmpty())
 	{
 		return;
 	}
 
-	if (!ModifierEditorSubsystem->RemoveActorsModifiers(Actors))
+	if (!ModifierSubsystem->RemoveActorsModifiers(Actors, true))
 	{
 		if (InData.Options.ShouldFireNotification())
 		{
@@ -544,21 +558,25 @@ void UE::ActorModifierCoreEditor::OnRemoveActorModifierMenuAction(const FActorMo
 
 void UE::ActorModifierCoreEditor::OnRemoveModifierMenuAction(const FActorModifierCoreEditorMenuData& InData)
 {
-	const UActorModifierCoreEditorSubsystem* ModifierEditorSubsystem = UActorModifierCoreEditorSubsystem::Get();
+	const UActorModifierCoreSubsystem* ModifierSubsystem = UActorModifierCoreSubsystem::Get();
 
-	if (!IsValid(ModifierEditorSubsystem))
+	if (!IsValid(ModifierSubsystem))
 	{
 		return;
 	}
 
-	TSet<UActorModifierCoreBase*> Modifiers = InData.Context.ContextModifiers;
+	TSet<UActorModifierCoreBase*> Modifiers = InData.Context.GetContextModifiers();
 	if (Modifiers.IsEmpty())
 	{
 		return;
 	}
 
 	FText OutFailReason;
-	if (!ModifierEditorSubsystem->RemoveModifiers(Modifiers, &OutFailReason))
+	FActorModifierCoreStackRemoveOp RemoveOp;
+	RemoveOp.bShouldTransact = true;
+	RemoveOp.FailReason = &OutFailReason;
+
+	if (!ModifierSubsystem->RemoveModifiers(Modifiers, RemoveOp))
 	{
 		if (InData.Options.ShouldFireNotification())
 		{
@@ -576,17 +594,22 @@ void UE::ActorModifierCoreEditor::OnRemoveSingleModifierMenuAction(UActorModifie
 	{
 		return;
 	}
-	
-	const UActorModifierCoreEditorSubsystem* ModifierEditorSubsystem = UActorModifierCoreEditorSubsystem::Get();
 
-	if (!IsValid(ModifierEditorSubsystem))
+	const UActorModifierCoreSubsystem* ModifierSubsystem = UActorModifierCoreSubsystem::Get();
+
+	if (!IsValid(ModifierSubsystem))
 	{
 		return;
 	}
 
-	TSet<UActorModifierCoreBase*> RemoveModifiers { InModifier };
 	FText OutFailReason;
-	if (!ModifierEditorSubsystem->RemoveModifiers(RemoveModifiers, &OutFailReason))
+	TSet<UActorModifierCoreBase*> RemoveModifiers { InModifier };
+
+	FActorModifierCoreStackRemoveOp RemoveOp;
+	RemoveOp.bShouldTransact = true;
+	RemoveOp.FailReason = &OutFailReason;
+
+	if (!ModifierSubsystem->RemoveModifiers(RemoveModifiers, RemoveOp))
 	{
 		if (InData.Options.ShouldFireNotification())
 		{
@@ -600,16 +623,23 @@ void UE::ActorModifierCoreEditor::OnRemoveSingleModifierMenuAction(UActorModifie
 
 void UE::ActorModifierCoreEditor::OnMoveModifierMenuAction(UActorModifierCoreBase* InMoveModifier, UActorModifierCoreBase* InPositionModifier, const FActorModifierCoreEditorMenuData& InData)
 {
-	const UActorModifierCoreEditorSubsystem* ModifierEditorSubsystem = UActorModifierCoreEditorSubsystem::Get();
+	const UActorModifierCoreSubsystem* ModifierSubsystem = UActorModifierCoreSubsystem::Get();
 
-	if (!IsValid(ModifierEditorSubsystem))
+	if (!IsValid(ModifierSubsystem) || !IsValid(InMoveModifier))
 	{
 		return;
 	}
-	
+
 	FText OutFailReason;
-	static constexpr EActorModifierCoreStackPosition MovePosition = EActorModifierCoreStackPosition::After;
-	if (!ModifierEditorSubsystem->MoveModifier(InMoveModifier, InPositionModifier, MovePosition, &OutFailReason))
+
+	FActorModifierCoreStackMoveOp MoveOp;
+	MoveOp.bShouldTransact = true;
+	MoveOp.FailReason = &OutFailReason;
+	MoveOp.MoveModifier = InMoveModifier;
+	MoveOp.MovePosition = EActorModifierCoreStackPosition::After;
+	MoveOp.MovePositionContext = InPositionModifier;
+
+	if (!ModifierSubsystem->MoveModifier(InMoveModifier->GetModifierStack(), MoveOp))
 	{
 		if (InData.Options.ShouldFireNotification())
 		{
@@ -623,9 +653,9 @@ void UE::ActorModifierCoreEditor::OnMoveModifierMenuAction(UActorModifierCoreBas
 
 void UE::ActorModifierCoreEditor::OnEnableModifierMenuAction(const FActorModifierCoreEditorMenuData& InData, bool bInEnable, bool bInStack)
 {
-	const UActorModifierCoreEditorSubsystem* ModifierEditorSubsystem = UActorModifierCoreEditorSubsystem::Get();
+	const UActorModifierCoreSubsystem* ModifierSubsystem = UActorModifierCoreSubsystem::Get();
 
-	if (!IsValid(ModifierEditorSubsystem))
+	if (!IsValid(ModifierSubsystem))
 	{
 		return;
 	}
@@ -633,11 +663,11 @@ void UE::ActorModifierCoreEditor::OnEnableModifierMenuAction(const FActorModifie
 	TSet<UActorModifierCoreBase*> Modifiers;
 	if (bInStack)
 	{
-		Modifiers = InData.Context.ContextStacks;
+		Modifiers = InData.Context.GetContextStacks();
 	}
 	else
 	{
-		Modifiers = InData.Context.ContextModifiers;
+		Modifiers = InData.Context.GetContextModifiers();
 	}
 
 	if (Modifiers.IsEmpty())
@@ -645,13 +675,13 @@ void UE::ActorModifierCoreEditor::OnEnableModifierMenuAction(const FActorModifie
 		return;
 	}
 
-	if (!ModifierEditorSubsystem->EnableModifiers(bInEnable, Modifiers))
+	if (!ModifierSubsystem->EnableModifiers(Modifiers, bInEnable, true))
 	{
 		if (InData.Options.ShouldFireNotification())
 		{
 			static const FText EnableStacksWarning(LOCTEXT("EnableStacks", "Could not update stack(s) state"));
 			static const FText EnableModifiersWarning(LOCTEXT("EnableModifiers", "Could not update modifier(s) state"));
-			
+
 			FNotificationInfo NotificationInfo(bInStack ? EnableStacksWarning : EnableModifiersWarning);
 			NotificationInfo.ExpireDuration = 3.f;
 			NotificationInfo.bFireAndForget = true;
diff --git a/Engine/Plugins/Experimental/ActorModifierCore/Source/ActorModifierCoreEditor/Private/Modifiers/Customizations/ActorModifierCoreEditorStackCustomization.cpp b/Engine/Plugins/Experimental/ActorModifierCore/Source/ActorModifierCoreEditor/Private/Modifiers/Customizations/ActorModifierCoreEditorStackCustomization.cpp
--- a/Engine/Plugins/Experimental/ActorModifierCore/Source/ActorModifierCoreEditor/Private/Modifiers/Customizations/ActorModifierCoreEditorStackCustomization.cpp
+++ b/Engine/Plugins/Experimental/ActorModifierCore/Source/ActorModifierCoreEditor/Private/Modifiers/Customizations/ActorModifierCoreEditorStackCustomization.cpp
@@ -97,7 +97,7 @@ void UActorModifierCoreEditorStackCustomization::CustomizeStackHeader(const FOpe
 
 				if (Modifier && !Modifier->IsModifierStack())
 				{
-					PinnedKeywords.Add(Modifier->GetModifierMetadata().GetCategory().ToString());
+					PinnedKeywords.Add(Modifier->GetModifierCategory().ToString());
 				}
 			}
 		}
@@ -150,8 +150,8 @@ void UActorModifierCoreEditorStackCustomization::CustomizeItemHeader(const FOper
 		// Item keyword for search
 		const TSet<FString> SearchKeywords
 		{
-			Modifier->GetModifierMetadata().GetDisplayName().ToString(),
-			Modifier->GetModifierMetadata().GetCategory().ToString()
+			Modifier->GetModifierName().ToString(),
+			Modifier->GetModifierCategory().ToString()
 		};
 
 		FSlateIcon ModifierIcon = FSlateIconFinder::FindIconForClass(UActorModifierCoreBase::StaticClass());
@@ -292,7 +292,7 @@ bool UActorModifierCoreEditorStackCustomization::OnIsItemDraggable(const FOperat
 
 TOptional<EItemDropZone> UActorModifierCoreEditorStackCustomization::OnItemCanAcceptDrop(const TArray<FOperatorStackEditorItemPtr>& InDraggedItems, const FOperatorStackEditorItemPtr& InDropZoneItem, EItemDropZone InZone)
 {
-	const UActorModifierCoreEditorSubsystem* ExtensionSubsystem = UActorModifierCoreEditorSubsystem::Get();
+	const UActorModifierCoreSubsystem* ModifierSubsystem = UActorModifierCoreSubsystem::Get();
 
 	TSet<UActorModifierCoreBase*> DraggedModifiers;
 	for (const FOperatorStackEditorItemPtr& Item : InDraggedItems)
@@ -324,7 +324,7 @@ TOptional<EItemDropZone> UActorModifierCoreEditorStackCustomization::OnItemCanAc
 		TArray<UActorModifierCoreBase*> CloneModifiers;
 
 		const EActorModifierCoreStackPosition Position = InZone == EItemDropZone::AboveItem ? EActorModifierCoreStackPosition::Before : EActorModifierCoreStackPosition::After;
-		ExtensionSubsystem->GetSortedModifiers(DraggedModifiers, DropModifier->GetModifiedActor(), DropModifier, Position, MoveModifiers, CloneModifiers);
+		ModifierSubsystem->GetSortedModifiers(DraggedModifiers, DropModifier->GetModifiedActor(), DropModifier, Position, MoveModifiers, CloneModifiers);
 
 		if (!MoveModifiers.IsEmpty())
 		{
@@ -338,7 +338,6 @@ TOptional<EItemDropZone> UActorModifierCoreEditorStackCustomization::OnItemCanAc
 void UActorModifierCoreEditorStackCustomization::OnDropItem(const TArray<FOperatorStackEditorItemPtr>& InDraggedItems, const FOperatorStackEditorItemPtr& InDropZoneItem, EItemDropZone InZone)
 {
 	const UActorModifierCoreSubsystem* const ModifierSubsystem = UActorModifierCoreSubsystem::Get();
-	const UActorModifierCoreEditorSubsystem* ExtensionSubsystem = UActorModifierCoreEditorSubsystem::Get();
 
 	TSet<UActorModifierCoreBase*> Modifiers;
 	for (const FOperatorStackEditorItemPtr& Item : InDraggedItems)
@@ -366,7 +365,7 @@ void UActorModifierCoreEditorStackCustomization::OnDropItem(const TArray<FOperat
 	{
 		UActorModifierCoreBase* DropModifier = InDropZoneItem->Get<UActorModifierCoreBase>();
 
-		const UActorModifierCoreStack* const TargetStack = DropModifier->GetModifierStack();
+		UActorModifierCoreStack* const TargetStack = DropModifier->GetModifierStack();
 
 		if (!IsValid(ModifierSubsystem) || !IsValid(TargetStack))
 		{
@@ -379,7 +378,7 @@ void UActorModifierCoreEditorStackCustomization::OnDropItem(const TArray<FOperat
 		TArray<UActorModifierCoreBase*> CloneModifiers;
 
 		const EActorModifierCoreStackPosition Position = InZone == EItemDropZone::AboveItem ? EActorModifierCoreStackPosition::Before : EActorModifierCoreStackPosition::After;
-		ExtensionSubsystem->GetSortedModifiers(Modifiers, TargetActor, DropModifier, Position, MoveModifiers, CloneModifiers);
+		ModifierSubsystem->GetSortedModifiers(Modifiers, TargetActor, DropModifier, Position, MoveModifiers, CloneModifiers);
 
 		if (MoveModifiers.IsEmpty())
 		{
@@ -387,7 +386,13 @@ void UActorModifierCoreEditorStackCustomization::OnDropItem(const TArray<FOperat
 		}
 
 		FText FailReason;
-		ExtensionSubsystem->MoveModifiers(MoveModifiers, DropModifier, Position, &FailReason);
+		FActorModifierCoreStackMoveOp MoveOp;
+		MoveOp.bShouldTransact = true;
+		MoveOp.FailReason = &FailReason;
+		MoveOp.MovePosition = Position;
+		MoveOp.MovePositionContext = DropModifier;
+
+		ModifierSubsystem->MoveModifiers(MoveModifiers, TargetStack, MoveOp);
 
 		if (!FailReason.IsEmpty())
 		{
@@ -582,9 +587,9 @@ void UActorModifierCoreEditorStackCustomization::RemoveModifierAction(UActorModi
 		return;
 	}
 
-	const UActorModifierCoreEditorSubsystem* ModifierEditorSubsystem = UActorModifierCoreEditorSubsystem::Get();
+	const UActorModifierCoreSubsystem* ModifierSubsystem = UActorModifierCoreSubsystem::Get();
 
-	if (!IsValid(ModifierEditorSubsystem))
+	if (!IsValid(ModifierSubsystem))
 	{
 		return;
 	}
@@ -596,7 +601,7 @@ void UActorModifierCoreEditorStackCustomization::RemoveModifierAction(UActorModi
 			InModifier->GetModifiedActor()
 		};
 
-		ModifierEditorSubsystem->RemoveActorsModifiers(Actors);
+		ModifierSubsystem->RemoveActorsModifiers(Actors, true);
 	}
 	else
 	{
@@ -604,7 +609,11 @@ void UActorModifierCoreEditorStackCustomization::RemoveModifierAction(UActorModi
 		Modifiers.Add(InModifier);
 
 		FText OutFailReason;
-		if (!ModifierEditorSubsystem->RemoveModifiers(Modifiers, &OutFailReason))
+		FActorModifierCoreStackRemoveOp RemoveOp;
+		RemoveOp.bShouldTransact = true;
+		RemoveOp.FailReason = &OutFailReason;
+
+		if (!ModifierSubsystem->RemoveModifiers(Modifiers, RemoveOp))
 		{
 			FNotificationInfo NotificationInfo(OutFailReason);
 			NotificationInfo.ExpireDuration = 3.0f;
@@ -889,8 +898,8 @@ bool UActorModifierCoreEditorStackCustomization::UpdateModifierFromPropertiesHan
 
 bool UActorModifierCoreEditorStackCustomization::AddModifierFromClipboard(TSet<AActor*>& InActors) const
 {
-	const UActorModifierCoreEditorSubsystem* ModifierEditorSubsystem = UActorModifierCoreEditorSubsystem::Get();
-	if (!IsValid(ModifierEditorSubsystem))
+	const UActorModifierCoreSubsystem* ModifierSubsystem = UActorModifierCoreSubsystem::Get();
+	if (!IsValid(ModifierSubsystem))
 	{
 		return false;
 	}
@@ -901,7 +910,14 @@ bool UActorModifierCoreEditorStackCustomization::AddModifierFromClipboard(TSet<A
 		const FName ModifierName = ModifierPropertiesWrapper.ModifierName;
 
 		FText OutFailReason;
-		if (!ModifierEditorSubsystem->AddActorsModifiers(ModifierName, InActors, &OutFailReason))
+		FActorModifierCoreStackInsertOp AddOp;
+		AddOp.bShouldTransact = true;
+		AddOp.FailReason = &OutFailReason;
+		AddOp.NewModifierName = ModifierName;
+
+		ModifierSubsystem->AddActorsModifiers(InActors, AddOp);
+
+		if (!OutFailReason.IsEmpty())
 		{
 			FNotificationInfo NotificationInfo(OutFailReason);
 			NotificationInfo.ExpireDuration = 3.0f;
diff --git a/Engine/Plugins/Experimental/ActorModifierCore/Source/ActorModifierCoreEditor/Private/Subsystems/ActorModifierCoreEditorSubsystem.cpp b/Engine/Plugins/Experimental/ActorModifierCore/Source/ActorModifierCoreEditor/Private/Subsystems/ActorModifierCoreEditorSubsystem.cpp
--- a/Engine/Plugins/Experimental/ActorModifierCore/Source/ActorModifierCoreEditor/Private/Subsystems/ActorModifierCoreEditorSubsystem.cpp
+++ b/Engine/Plugins/Experimental/ActorModifierCore/Source/ActorModifierCoreEditor/Private/Subsystems/ActorModifierCoreEditorSubsystem.cpp
@@ -37,7 +37,7 @@ void UActorModifierCoreEditorSubsystem::Initialize(FSubsystemCollectionBase& Col
 	EngineSubsystem = UActorModifierCoreSubsystem::Get();
 
 	check(EngineSubsystem.IsValid());
-	
+
 	// Initialize here if not called anywhere to setup ClassIcons
 	FActorModifierCoreEditorStyle::Get();
 }
@@ -47,409 +47,19 @@ void UActorModifierCoreEditorSubsystem::Deinitialize()
 	Super::Deinitialize();
 }
 
-bool UActorModifierCoreEditorSubsystem::EnableModifiers(bool bInEnabled, TSet<UActorModifierCoreBase*>& InModifiers, bool bInShouldTransact) const
-{
-	if (InModifiers.IsEmpty())
-	{
-		return false;
-	}
-
-	// create transaction
-	FText TransactionText;
-	if (bInEnabled)
-	{
-		TransactionText = FText::Format(
-			LOCTEXT("EnableModifiers.Enable", "Enabling {0} modifier(s)"),
-			FText::FromString(FString::FromInt(InModifiers.Num())));
-	}
-	else
-	{
-		TransactionText = FText::Format(
-			LOCTEXT("EnableModifiers.Disable", "Disabling {0} modifier(s)"),
-			FText::FromString(FString::FromInt(InModifiers.Num())));
-	}
-	FScopedTransaction Transaction(TransactionText, bInShouldTransact);
-	
-	// lets group modifiers to batch the operation, better to update the stack only once instead of many times
-	while (!InModifiers.IsEmpty())
-	{
-		TSet<UActorModifierCoreBase*> StackModifiers;
-
-		const UActorModifierCoreStack* CurrentStack = nullptr;
-	
-		for (UActorModifierCoreBase* Modifier : InModifiers)
-		{
-			if (!IsValid(Modifier))
-			{
-				continue;
-			}
-
-			UActorModifierCoreStack* ModifierStack = Modifier->GetModifierStack();
-			ModifierStack = IsValid(ModifierStack) ? ModifierStack : Modifier->GetRootModifierStack();
-			
-			if (!CurrentStack)
-			{
-				CurrentStack = ModifierStack;
-			}
-
-			if (ModifierStack == CurrentStack)
-			{
-				StackModifiers.Add(Modifier);
-			}
-		}
-
-		{
-			// disable stack modifiers update
-			FActorModifierCoreScopedLock Lock(StackModifiers);
-			for (UActorModifierCoreBase* Modifier : StackModifiers)
-			{
-				Modifier->SetModifierEnabled(bInEnabled);
-			}
-		}
-
-		InModifiers = InModifiers.Difference(StackModifiers);
-	}
-
-	return true;
-}
-
-bool UActorModifierCoreEditorSubsystem::RemoveModifiers(TSet<UActorModifierCoreBase*>& InModifiers, FText* OutFailReason, bool bInShouldTransact) const
-{
-	if (InModifiers.IsEmpty())
-	{
-		return false;
-	}
-
-	// create transaction
-	FText TransactionText;
-	const UActorModifierCoreBase* SingleModifier = InModifiers.Array()[0];
-	if (InModifiers.Num() == 1 && IsValid(SingleModifier))
-	{
-		TransactionText = FText::Format(
-			LOCTEXT("RemoveSingleModifiers", "Removing {0} modifier"),
-			FText::FromName(SingleModifier->GetModifierName()));
-	}
-	else
-	{
-		TransactionText = FText::Format(
-			LOCTEXT("RemoveMultipleModifiers", "Removing {0} modifier(s)"),
-			FText::FromString(FString::FromInt(InModifiers.Num())));
-	}
-	FScopedTransaction Transaction(TransactionText, bInShouldTransact);
-
-	// lets group modifiers to batch the operation, better to update the stack only once instead of many times
-	while (!InModifiers.IsEmpty())
-	{
-		TSet<UActorModifierCoreBase*> StackModifiers;
-
-		UActorModifierCoreStack* CurrentStack = nullptr;
-	
-		for (UActorModifierCoreBase* Modifier : InModifiers)
-		{
-			if (!IsValid(Modifier))
-			{
-				continue;
-			}
-
-			UActorModifierCoreStack* ModifierStack = Modifier->GetModifierStack();
-			ModifierStack = IsValid(ModifierStack) ? ModifierStack : Modifier->GetRootModifierStack();
-			
-			if (!CurrentStack)
-			{
-				CurrentStack = ModifierStack;
-			}
-
-			if (ModifierStack == CurrentStack)
-			{
-				StackModifiers.Add(Modifier);
-			}
-		}
-
-		// Sort by their order in stack first to avoid dependencies errors
-		bool bSuccess = true;
-		StackModifiers.Sort([CurrentStack](UActorModifierCoreBase& A, UActorModifierCoreBase& B)
-		{
-			bool bSortResult = true; 
-			CurrentStack->ProcessFunction([&A, &B, &bSortResult](const UActorModifierCoreBase* InModifier)->bool
-			{
-				if (InModifier == &A || InModifier == &B)
-				{
-					bSortResult = InModifier == &A ? false : true;
-					return false;
-				}
-				return true;
-			});
-			return bSortResult;
-		});
-
-		CurrentStack->ProcessLockFunction([this, &StackModifiers, OutFailReason, CurrentStack, &bSuccess]()
-		{
-			for (UActorModifierCoreBase* RemoveModifier : StackModifiers)
-			{
-				FActorModifierCoreStackRemoveOp RemoveOp;
-				RemoveOp.RemoveModifier = RemoveModifier;
-				RemoveOp.FailReason = OutFailReason;
-				if (!CurrentStack->RemoveModifier(RemoveOp))
-				{
-					bSuccess = false;
-					break;
-				}
-			}
-		});
-		
-		if (!bSuccess)
-		{
-			return false;
-		}
-
-		InModifiers = InModifiers.Difference(StackModifiers);
-	}
-
-	return true;
-}
-
-bool UActorModifierCoreEditorSubsystem::RemoveActorsModifiers(TSet<AActor*>& InActors, bool bInShouldTransact) const
-{
-	if (InActors.IsEmpty())
-	{
-		return false;
-	}
-
-	// Get actors stack
-	TArray<UActorModifierCoreStack*> ActorStacks;
-	for (const AActor* Actor : InActors)
-	{
-		if (!IsValid(Actor))
-		{
-			continue;
-		}
-
-		if (UActorModifierCoreStack* ActorStack = EngineSubsystem->GetActorModifierStack(Actor))
-		{
-			ActorStacks.Add(ActorStack);
-		}
-	}
-
-	// create transaction
-	const FText TransactionText = FText::Format(
-		LOCTEXT("RemoveModifier", "Remove all modifiers from {0} actor(s)"),
-		FText::FromString(FString::FromInt(InActors.Num())));
-	FScopedTransaction Transaction(TransactionText, bInShouldTransact);
-
-	// remove modifiers from actors
-	for (UActorModifierCoreStack* ActorStack : ActorStacks)
-	{
-		if (!ActorStack->RemoveAllModifiers())
-		{
-			return false;
-		}
-		// remove to indicate it was done
-		InActors.Remove(ActorStack->GetModifiedActor());
-	}
-	
-	return true;
-}
-
-bool UActorModifierCoreEditorSubsystem::AddActorsModifiers(const FName& InModifierName, TSet<AActor*>& InActors, FText* OutFailReason, bool bInShouldTransact) const
-{
-	if (!EngineSubsystem.IsValid())
-	{
-		return false;
-	}
-
-	if (!EngineSubsystem->IsRegisteredModifierClass(InModifierName))
-	{
-		return false;
-	}
-	
-	// Get actors stack, create if none
-	TArray<UActorModifierCoreStack*> ActorStacks;
-	for (AActor* Actor : InActors)
-	{
-		if (!IsValid(Actor))
-		{
-			continue;
-		}
-
-		UActorModifierCoreStack* ActorStack = EngineSubsystem->GetActorModifierStack(Actor);
-		
-		if (!ActorStack)
-		{
-			ActorStack = EngineSubsystem->AddActorModifierStack(Actor);
-		}
-
-		if (IsValid(ActorStack))
-		{
-			ActorStacks.Add(ActorStack);
-		}
-	}
-
-	// create transaction
-	const FText TransactionText = FText::Format(
-		LOCTEXT("AddModifier", "Add {0} modifier on {1} actor(s)"),
-		FText::FromName(InModifierName),
-		FText::FromString(FString::FromInt(InActors.Num())));
-	FScopedTransaction Transaction(TransactionText, bInShouldTransact);
-
-	// add modifier to actors
-	FActorModifierCoreStackInsertOp InsertOp;
-	InsertOp.NewModifierName = InModifierName;
-	InsertOp.FailReason = OutFailReason;
-	
-	for (UActorModifierCoreStack* ActorStack : ActorStacks)
-	{
-		if (!ActorStack->InsertModifier(InsertOp))
-		{
-			return false;
-		}
-		
-		// remove to indicate it was done
-		InActors.Remove(ActorStack->GetModifiedActor());
-	}
-	
-	return true;
-}
-
-bool UActorModifierCoreEditorSubsystem::InsertModifier(const FName& InModifierName, UActorModifierCoreStack* InStack, UActorModifierCoreBase* InPositionModifier, EActorModifierCoreStackPosition InPosition, FText* OutFailReason, bool bInShouldTransact) const
-{
-	if (!EngineSubsystem.IsValid())
-	{
-		return false;
-	}
-
-	if (!EngineSubsystem->IsRegisteredModifierClass(InModifierName))
-	{
-		return false;
-	}
-
-	if (!IsValid(InStack))
-	{
-		return false;
-	}
-
-	// create transaction
-	FText TransactionText;
-	if (InPosition == EActorModifierCoreStackPosition::Before)
-	{
-		if (IsValid(InPositionModifier))
-		{
-			TransactionText = FText::Format(
-				LOCTEXT("InsertModifierBefore", "Insert {0} modifier before {1}"),
-				FText::FromName(InModifierName),
-				FText::FromName(InPositionModifier->GetModifierName()));
-		}
-		else
-		{
-			TransactionText = FText::Format(
-				LOCTEXT("InsertModifierAtEnd", "Insert {0} modifier at the end of stack"),
-				FText::FromName(InModifierName));
-		}
-	}
-	else if (InPosition == EActorModifierCoreStackPosition::After)
-	{
-		if (IsValid(InPositionModifier))
-		{
-			TransactionText = FText::Format(
-				LOCTEXT("InsertModifierAfter", "Insert {0} modifier after {1}"),
-				FText::FromName(InModifierName),
-				FText::FromName(InPositionModifier->GetModifierName()));
-		}
-		else
-		{
-			TransactionText = FText::Format(
-				LOCTEXT("InsertModifierAtStart", "Insert {0} modifier at the start of stack"),
-				FText::FromName(InModifierName));
-		}
-	}
-	FScopedTransaction Transaction(TransactionText, bInShouldTransact);
-
-	// insert modifier in stack
-	FActorModifierCoreStackInsertOp InsertOp;
-	InsertOp.NewModifierName = InModifierName;
-	InsertOp.InsertPosition = InPosition;
-	InsertOp.InsertPositionContext = InPositionModifier;
-	InsertOp.FailReason = OutFailReason;
-	return InStack->InsertModifier(InsertOp) != nullptr;
-}
-
-bool UActorModifierCoreEditorSubsystem::MoveModifier(UActorModifierCoreBase* InMoveModifier, UActorModifierCoreBase* InPositionModifier, EActorModifierCoreStackPosition InPosition, FText* OutFailReason, bool bInShouldTransact) const
-{
-	if (!EngineSubsystem.IsValid())
-	{
-		return false;
-	}
-
-	if (!IsValid(InMoveModifier))
-	{
-		return false;
-	}
-
-	UActorModifierCoreStack* ModifierStack = InMoveModifier->GetModifierStack();
-	if (!IsValid(ModifierStack))
-	{
-		return false;
-	}
-	
-	const FName& MoveModifierName = InMoveModifier->GetModifierName();
-	
-	// create transaction
-	FText TransactionText;
-	if (InPosition == EActorModifierCoreStackPosition::Before)
-	{
-		if (IsValid(InPositionModifier))
-		{
-			TransactionText = FText::Format(
-				LOCTEXT("MoveModifierBefore", "Move {0} modifier before {1}"),
-				FText::FromName(MoveModifierName),
-				FText::FromName(InPositionModifier->GetModifierName()));
-		}
-		else
-		{
-			TransactionText = FText::Format(
-				LOCTEXT("MoveModifierToEnd", "Move {0} modifier at the end of stack"),
-				FText::FromName(MoveModifierName));
-		}
-	}
-	else if (InPosition == EActorModifierCoreStackPosition::After)
-	{
-		if (IsValid(InPositionModifier))
-		{
-			TransactionText = FText::Format(
-				LOCTEXT("MoveModifierAfter", "Move {0} modifier after {1}"),
-				FText::FromName(MoveModifierName),
-				FText::FromName(InPositionModifier->GetModifierName()));
-		}
-		else
-		{
-			TransactionText = FText::Format(
-				LOCTEXT("MoveModifierToStart", "Move {0} modifier at the start of stack"),
-				FText::FromName(MoveModifierName));
-		}
-	}
-	FScopedTransaction Transaction(TransactionText, bInShouldTransact);
-
-	// move modifier in stack
-	FActorModifierCoreStackMoveOp MoveOp;
-	MoveOp.MoveModifier = InMoveModifier;
-	MoveOp.MovePosition = InPosition;
-	MoveOp.MovePositionContext = InPositionModifier;
-	MoveOp.FailReason = OutFailReason;
-	return ModifierStack->MoveModifier(MoveOp);
-}
-
 bool UActorModifierCoreEditorSubsystem::FillModifierMenu(UToolMenu* InMenu, const FActorModifierCoreEditorMenuContext& InContext, const FActorModifierCoreEditorMenuOptions& InMenuOptions) const
 {
 	if (!InMenu || InContext.IsEmpty())
 	{
 		return false;
 	}
-	
+
 	const FActorModifierCoreEditorMenuData MenuData(InContext, InMenuOptions);
 	FToolMenuSection* ContextModifiersSection = nullptr;
 	if (InMenuOptions.ShouldCreateSubMenu())
 	{
 		static const FName ModifierSection("ContextModifierActions");
-	
+
 		ContextModifiersSection = InMenu->FindSection(ModifierSection);
 		if (!ContextModifiersSection)
 		{
@@ -458,7 +68,7 @@ bool UActorModifierCoreEditorSubsystem::FillModifierMenu(UToolMenu* InMenu, cons
 				, FToolMenuInsert(NAME_None, EToolMenuInsertType::First));
 		}
 	}
-	
+
 	if (InMenuOptions.GetMenuType() == EActorModifierCoreEditorMenuType::Add)
 	{
 		if (InContext.ContainsAnyActor())
@@ -500,7 +110,7 @@ bool UActorModifierCoreEditorSubsystem::FillModifierMenu(UToolMenu* InMenu, cons
 	}
 	else if (InMenuOptions.GetMenuType() == EActorModifierCoreEditorMenuType::Move)
 	{
-		if (InContext.ContainsOnlyModifier() && InContext.ContextModifiers.Num() == 1)
+		if (InContext.ContainsOnlyModifier() && InContext.GetContextModifiers().Num() == 1)
 		{
 			if (InMenuOptions.ShouldCreateSubMenu())
 			{
@@ -560,7 +170,7 @@ bool UActorModifierCoreEditorSubsystem::FillModifierMenu(UToolMenu* InMenu, cons
 	}
 	else if (InMenuOptions.GetMenuType() == EActorModifierCoreEditorMenuType::InsertBefore)
 	{
-		if (InContext.ContainsOnlyModifier() && InContext.ContextModifiers.Num() == 1)
+		if (InContext.ContainsOnlyModifier() && InContext.GetContextModifiers().Num() == 1)
 		{
 			if (InMenuOptions.ShouldCreateSubMenu())
 			{
@@ -580,7 +190,7 @@ bool UActorModifierCoreEditorSubsystem::FillModifierMenu(UToolMenu* InMenu, cons
 	}
 	else if (InMenuOptions.GetMenuType() == EActorModifierCoreEditorMenuType::InsertAfter)
 	{
-		if (InContext.ContainsOnlyModifier() && InContext.ContextModifiers.Num() == 1)
+		if (InContext.ContainsOnlyModifier() && InContext.GetContextModifiers().Num() == 1)
 		{
 			if (InMenuOptions.ShouldCreateSubMenu())
 			{
@@ -602,228 +212,6 @@ bool UActorModifierCoreEditorSubsystem::FillModifierMenu(UToolMenu* InMenu, cons
 	return false;
 }
 
-void UActorModifierCoreEditorSubsystem::GetSortedModifiers(const TSet<UActorModifierCoreBase*>& InModifiers, AActor* InTargetActor, UActorModifierCoreBase* InTargetModifier, EActorModifierCoreStackPosition InPosition, TArray<UActorModifierCoreBase*>& OutMoveModifiers, TArray<UActorModifierCoreBase*>& OutCloneModifiers) const
-{
-	const UActorModifierCoreSubsystem* const ModifierSubsystem = UActorModifierCoreSubsystem::Get();
-	
-	if (!IsValid(ModifierSubsystem) || InModifiers.IsEmpty() || !IsValid(InTargetActor))
-	{
-		return;
-	}
-
-	for (UActorModifierCoreBase* Modifier : InModifiers)
-	{
-		if (!IsValid(Modifier))
-		{
-			continue;	
-		}
-
-		// it's a clone operation if target actor is different as modifier actor
-		if (Modifier->GetModifiedActor() != InTargetActor)
-		{
-			OutCloneModifiers.Add(Modifier);
-		}
-		// it's a move operation if target actor is same as modifier actor
-		else
-		{
-			OutMoveModifiers.Add(Modifier);
-		}
-	}
-	
-	// Remove unsupported modifiers when inserting where target modifier is nullptr
-	if (!InTargetModifier)
-	{
-		const TSet<FName> AllowedModifiers = ModifierSubsystem->GetAllowedModifiers(InTargetActor, InTargetModifier, InPosition);
-		
-		OutMoveModifiers.RemoveAll([&AllowedModifiers](UActorModifierCoreBase* InModifier)
-		{
-			return !IsValid(InModifier) || !AllowedModifiers.Contains(InModifier->GetModifierName());
-		});
-
-		OutCloneModifiers.RemoveAll([&AllowedModifiers](UActorModifierCoreBase* InModifier)
-		{
-			return !IsValid(InModifier) || !AllowedModifiers.Contains(InModifier->GetModifierName());
-		});
-	}
-
-	// Sort them by dependency and current order
-	OutMoveModifiers.Sort([InTargetModifier](const UActorModifierCoreBase& A, const UActorModifierCoreBase& B)
-	{
-		const FActorModifierCoreMetadata& MetadataA = A.GetModifierMetadata();
-		const FActorModifierCoreMetadata& MetadataB = B.GetModifierMetadata();
-		const bool bADependsOnB = MetadataA.DependsOn(MetadataB.GetName());
-
-		const UActorModifierCoreStack* StackA = A.GetModifierStack();
-		const UActorModifierCoreStack* StackB = B.GetModifierStack();
-		const bool bAIsBeforeB = StackA == StackB && StackA->ContainsModifierBefore(&A, &B);
-
-		bool bSortOrder = bAIsBeforeB && !bADependsOnB;
-		
-		const UActorModifierCoreBase* DependencyModifier = bADependsOnB ? &A : &B;
-		if (StackA->ContainsModifierAfter(InTargetModifier, DependencyModifier))
-		{
-			bSortOrder = !bSortOrder;
-		}
-			
-		return bSortOrder;
-	});
-
-	OutCloneModifiers.Sort([](const UActorModifierCoreBase& A, const UActorModifierCoreBase& B)
-	{
-		const FActorModifierCoreMetadata& MetadataA = A.GetModifierMetadata();
-		const FActorModifierCoreMetadata& MetadataB = B.GetModifierMetadata();
-		const bool bADependsOnB = MetadataA.DependsOn(MetadataB.GetName());
-
-		const UActorModifierCoreStack* StackA = A.GetModifierStack();
-		const UActorModifierCoreStack* StackB = B.GetModifierStack();
-		const bool bAIsBeforeB = StackA == StackB && StackA->ContainsModifierBefore(&A, &B);
-
-		const bool bSortOrder = bAIsBeforeB && !bADependsOnB;
-
-		return bSortOrder;
-	});
-}
-
-bool UActorModifierCoreEditorSubsystem::MoveModifiers(const TArray<UActorModifierCoreBase*>& InModifiers, UActorModifierCoreBase* InTargetModifier, EActorModifierCoreStackPosition InPosition, FText* OutFailReason, bool bInShouldTransact) const
-{
-	if (!InTargetModifier)
-	{
-		return false;
-	}
-	
-	const UActorModifierCoreSubsystem* const ModifierSubsystem = UActorModifierCoreSubsystem::Get();
-	UActorModifierCoreStack* const TargetStack = InTargetModifier->GetModifierStack();
-
-	if (!IsValid(ModifierSubsystem) || !IsValid(TargetStack) || InModifiers.IsEmpty())
-	{
-		return false;
-	}
-
-	static const FText TransactionText = LOCTEXT("MoveModifiers", "Moving {0} modifier(s) {1} modifier {2}");
-	FScopedTransaction Transaction(
-		FText::Format(TransactionText
-			, FText::FromString(FString::FromInt(InModifiers.Num()))
-			, FText::FromString(InPosition == EActorModifierCoreStackPosition::After ? TEXT("after") : TEXT("before"))
-			, FText::FromName(InTargetModifier->GetModifierName()))
-		, bInShouldTransact);
-	
-	uint32 EditModifierCount = 0;
-
-	TargetStack->ProcessLockFunction([TargetStack, &InModifiers, OutFailReason, InPosition, InTargetModifier, &EditModifierCount]()
-	{
-		UActorModifierCoreBase* OperationContext = nullptr;
-		
-		for (UActorModifierCoreBase* Modifier : InModifiers)
-		{
-			if (IsValid(Modifier))
-			{
-				FActorModifierCoreStackMoveOp MoveOp;
-				MoveOp.FailReason = OutFailReason;
-				MoveOp.MovePosition = InPosition;
-				MoveOp.MovePositionContext = InTargetModifier;
-				MoveOp.MoveModifier = Modifier;
-
-				// Eg 1: when moving [bend, subdivide] before target modifier X : X is after them in the stack, we want bend added before X and subdivide added before Bend
-				// Eg 2: when moving [subdivide, bend] after target modifier X : X is before them in the stack, we want subdivide after X and bend after subdivide
-				if (OperationContext &&
-					((MoveOp.MovePosition == EActorModifierCoreStackPosition::After && TargetStack->ContainsModifierAfter(Modifier, InTargetModifier)) ||
-					(MoveOp.MovePosition == EActorModifierCoreStackPosition::Before && TargetStack->ContainsModifierBefore(Modifier, InTargetModifier))))
-				{
-					MoveOp.MovePositionContext = OperationContext;
-				}
-				
-				if (TargetStack->MoveModifier(MoveOp))
-				{
-					++EditModifierCount;
-					OperationContext = Modifier;
-				}
-				else
-				{
-					const AActor* TargetActor = InTargetModifier->GetModifiedActor();
-					const FText& ErrorText = MoveOp.FailReason ? *MoveOp.FailReason : FText::GetEmpty();
-					
-					// Move modifiers on actor failed
-					UE_LOG(LogActorModifierCoreEditorSubsystem, Warning, TEXT("Move modifier %s on actor %s failed : %s"),
-							*MoveOp.MoveModifier->GetModifierName().ToString(),
-							*TargetActor->GetActorNameOrLabel(),
-							*ErrorText.ToString());
-					break;
-				}
-			}
-		}
-	});
-
-	return EditModifierCount > 0;
-}
-
-bool UActorModifierCoreEditorSubsystem::CloneModifiers(const TArray<UActorModifierCoreBase*>& InModifiers, UActorModifierCoreBase* InTargetModifier, EActorModifierCoreStackPosition InPosition, FText* OutFailReason, bool bInShouldTransact) const
-{
-	if (!InTargetModifier)
-	{
-		return false;
-	}
-	
-	const UActorModifierCoreSubsystem* const ModifierSubsystem = UActorModifierCoreSubsystem::Get();
-	UActorModifierCoreStack* const TargetStack = InTargetModifier->GetModifierStack();
-
-	if (!IsValid(ModifierSubsystem) || !IsValid(TargetStack) || InModifiers.IsEmpty())
-	{
-		return false;
-	}
-
-	static const FText TransactionText = LOCTEXT("CloneModifiers", "Cloning {0} modifier(s) {1} modifier {2}");
-	FScopedTransaction Transaction(
-		FText::Format(TransactionText
-			, FText::FromString(FString::FromInt(InModifiers.Num()))
-			, FText::FromString(InPosition == EActorModifierCoreStackPosition::After ? TEXT("after") : TEXT("before"))
-			, FText::FromName(InTargetModifier->GetModifierName()))
-		, bInShouldTransact);
-	
-	uint32 EditModifierCount = 0;
-
-	TargetStack->ProcessLockFunction([TargetStack, &InModifiers, OutFailReason, InPosition, InTargetModifier, &EditModifierCount]()
-	{
-		UActorModifierCoreBase* OperationContext = nullptr;
-		for (UActorModifierCoreBase* Modifier : InModifiers)
-		{
-			if (IsValid(Modifier))
-			{
-				FActorModifierCoreStackCloneOp CloneOp;
-				CloneOp.FailReason = OutFailReason;
-				CloneOp.ClonePosition = InPosition;
-				CloneOp.ClonePositionContext = InTargetModifier;
-				CloneOp.CloneModifier = Modifier;
-
-				// When inserting A, B where B depends on A after C, clone A after C then clone B after A
-				if (OperationContext && CloneOp.ClonePosition == EActorModifierCoreStackPosition::After)
-				{
-					CloneOp.ClonePositionContext = OperationContext;
-				}
-				
-				if (UActorModifierCoreBase* NewClonedModifier = TargetStack->CloneModifier(CloneOp))
-				{
-					++EditModifierCount;
-					OperationContext = NewClonedModifier;
-				}
-				else
-				{
-					const AActor* TargetActor = InTargetModifier->GetModifiedActor();
-					const FText& ErrorText = CloneOp.FailReason ? *CloneOp.FailReason : FText::GetEmpty();
-					
-					// Clone modifiers on actor failed
-					UE_LOG(LogActorModifierCoreEditorSubsystem, Warning, TEXT("Clone modifier %s on actor %s failed : %s"),
-							*CloneOp.CloneModifier->GetModifierName().ToString(),
-							*TargetActor->GetActorNameOrLabel(),
-							*ErrorText.ToString());
-					break;
-				}
-			}
-		}
-	});
-
-	return EditModifierCount > 0;
-}
-
 TSharedPtr<SActorModifierCoreEditorProfiler> UActorModifierCoreEditorSubsystem::CreateProfilerWidget(TSharedPtr<FActorModifierCoreProfiler> InProfiler)
 {
 	if (!InProfiler.IsValid())
diff --git a/Engine/Plugins/Experimental/ActorModifierCore/Source/ActorModifierCoreEditor/Public/Modifiers/ActorModifierCoreEditorMenuDefs.h b/Engine/Plugins/Experimental/ActorModifierCore/Source/ActorModifierCoreEditor/Public/Modifiers/ActorModifierCoreEditorMenuDefs.h
--- a/Engine/Plugins/Experimental/ActorModifierCore/Source/ActorModifierCoreEditor/Public/Modifiers/ActorModifierCoreEditorMenuDefs.h
+++ b/Engine/Plugins/Experimental/ActorModifierCore/Source/ActorModifierCoreEditor/Public/Modifiers/ActorModifierCoreEditorMenuDefs.h
@@ -47,7 +47,7 @@ struct ACTORMODIFIERCOREEDITOR_API FActorModifierCoreEditorMenuOptions
 private:
 	/** What type of menu should be generated */
 	EActorModifierCoreEditorMenuType MenuType = EActorModifierCoreEditorMenuType::Add;
-	
+
 	/** Create a sub-menu to display actions instead of filling on the same level
 	 * with submenu : Root/Action/Choices... without : Root/Choices...
 	 */
@@ -63,12 +63,23 @@ private:
 /** Used to cache context results in case you call menu function with different menu options */
 struct ACTORMODIFIERCOREEDITOR_API FActorModifierCoreEditorMenuContext
 {
-	TSet<AActor*> ContextActors;
-	TSet<UActorModifierCoreBase*> ContextModifiers;
-	TSet<UActorModifierCoreBase*> ContextStacks;
-
 	explicit FActorModifierCoreEditorMenuContext(const TSet<TWeakObjectPtr<UObject>>& InContextObjects);
 
+	const TSet<AActor*>& GetContextActors() const
+	{
+		return ContextActors;
+	}
+
+	const TSet<UActorModifierCoreBase*>& GetContextModifiers() const
+	{
+		return ContextModifiers;
+	}
+
+	const TSet<UActorModifierCoreBase*>& GetContextStacks() const
+	{
+		return ContextStacks;
+	}
+
 	bool IsEmpty() const;
 	bool ContainsAnyActor() const;
 	bool ContainsAnyModifier() const;
@@ -79,16 +90,21 @@ struct ACTORMODIFIERCOREEDITOR_API FActorModifierCoreEditorMenuContext
 	bool ContainsDisabledStack() const;
 	bool ContainsEnabledModifier() const;
 	bool ContainsEnabledStack() const;
+
+protected:
+	TSet<AActor*> ContextActors;
+	TSet<UActorModifierCoreBase*> ContextModifiers;
+	TSet<UActorModifierCoreBase*> ContextStacks;
 };
 
 /** Used internally to group menu data together */
 struct FActorModifierCoreEditorMenuData
 {
-	const FActorModifierCoreEditorMenuContext Context;
-	const FActorModifierCoreEditorMenuOptions Options;
-
 	FActorModifierCoreEditorMenuData(const FActorModifierCoreEditorMenuContext& InContext, const FActorModifierCoreEditorMenuOptions& InOptions)
 		: Context(InContext)
 		, Options(InOptions)
 	{}
+
+	const FActorModifierCoreEditorMenuContext Context;
+	const FActorModifierCoreEditorMenuOptions Options;
 };
diff --git a/Engine/Plugins/Experimental/ActorModifierCore/Source/ActorModifierCoreEditor/Public/Subsystems/ActorModifierCoreEditorSubsystem.h b/Engine/Plugins/Experimental/ActorModifierCore/Source/ActorModifierCoreEditor/Public/Subsystems/ActorModifierCoreEditorSubsystem.h
--- a/Engine/Plugins/Experimental/ActorModifierCore/Source/ActorModifierCoreEditor/Public/Subsystems/ActorModifierCoreEditorSubsystem.h
+++ b/Engine/Plugins/Experimental/ActorModifierCore/Source/ActorModifierCoreEditor/Public/Subsystems/ActorModifierCoreEditorSubsystem.h
@@ -19,51 +19,24 @@ class UActorModifierCoreSubsystem;
 
 /** Singleton class that handles editor operations for modifiers */
 UCLASS()
-class ACTORMODIFIERCOREEDITOR_API UActorModifierCoreEditorSubsystem : public UEditorSubsystem
+class UActorModifierCoreEditorSubsystem : public UEditorSubsystem
 {
 	GENERATED_BODY()
 
 	friend class SActorModifierCoreEditorProfiler;
-	
+
 public:
 	UActorModifierCoreEditorSubsystem();
-	
-	static UActorModifierCoreEditorSubsystem* Get();
-	
+
+	ACTORMODIFIERCOREEDITOR_API static UActorModifierCoreEditorSubsystem* Get();
+
 	//~ Begin USubsystem
 	virtual void Initialize(FSubsystemCollectionBase& Collection) override;
 	virtual void Deinitialize() override;
 	//~ End USubsystem
 
-	/** Enable or disable modifiers, will update the original array and perform a transaction if wanted */
-	bool EnableModifiers(bool bInEnabled, TSet<UActorModifierCoreBase*>& InModifiers, bool bInShouldTransact = true) const;
-	
-	/** Remove modifiers from different actors or stacks, will update the original array and perform a transaction if wanted */
-	bool RemoveModifiers(TSet<UActorModifierCoreBase*>& InModifiers, FText* OutFailReason = nullptr, bool bInShouldTransact = true) const;
-
-	/** Remove all modifiers from actors, will update the original array and perform a transaction if wanted */
-	bool RemoveActorsModifiers(TSet<AActor*>& InActors, bool bInShouldTransact = true) const;
-
-	/** Add a modifier to actors, adds a stack automatically if none is found, will update the original array and perform a transaction if wanted */
-	bool AddActorsModifiers(const FName& InModifierName, TSet<AActor*>& InActors, FText* OutFailReason = nullptr, bool bInShouldTransact = true) const;
-
-	/** Insert a modifier in a stack before or after another modifier, will perform a transaction if wanted */
-	bool InsertModifier(const FName& InModifierName, UActorModifierCoreStack* InStack, UActorModifierCoreBase* InPositionModifier, EActorModifierCoreStackPosition InPosition, FText* OutFailReason = nullptr, bool bInShouldTransact = true) const;
-
-	/** Moves a modifier in a stack before or after another modifier, will perform a transaction if wanted */
-	bool MoveModifier(UActorModifierCoreBase* InMoveModifier, UActorModifierCoreBase* InPositionModifier, EActorModifierCoreStackPosition InPosition, FText* OutFailReason = nullptr, bool bInShouldTransact = true) const;
-
 	/** Fills a menu based on context objects and menu options, will perform menu action transaction if wanted, returns true if the menu has been modified */
-	bool FillModifierMenu(UToolMenu* InMenu, const FActorModifierCoreEditorMenuContext& InContext, const FActorModifierCoreEditorMenuOptions& InMenuOptions) const;
-
-	/** Get order dependent modifiers in the correct order for stack operations */
-	void GetSortedModifiers(const TSet<UActorModifierCoreBase*>& InModifiers, AActor* InTargetActor, UActorModifierCoreBase* InTargetModifier, EActorModifierCoreStackPosition InPosition, TArray<UActorModifierCoreBase*>& OutMoveModifiers, TArray<UActorModifierCoreBase*>& OutCloneModifiers) const;
-
-	/** Move modifiers to a specific target modifier position, will perform menu action transaction if wanted */
-	bool MoveModifiers(const TArray<UActorModifierCoreBase*>& InModifiers, UActorModifierCoreBase* InTargetModifier, EActorModifierCoreStackPosition InPosition, FText* OutFailReason = nullptr, bool bInShouldTransact = true) const;
-
-	/** Clone modifiers to a specific target modifier position, will perform menu action transaction if wanted */
-	bool CloneModifiers(const TArray<UActorModifierCoreBase*>& InModifiers, UActorModifierCoreBase* InTargetModifier, EActorModifierCoreStackPosition InPosition, FText* OutFailReason = nullptr, bool bInShouldTransact = true) const;
+	ACTORMODIFIERCOREEDITOR_API bool FillModifierMenu(UToolMenu* InMenu, const FActorModifierCoreEditorMenuContext& InContext, const FActorModifierCoreEditorMenuOptions& InMenuOptions) const;
 
 	/** Register a profiler widget based on a modifier profiler class */
 	template<typename InProfilerClass, typename = typename TEnableIf<TIsDerivedFrom<InProfilerClass, FActorModifierCoreProfiler>::Value>::Type
@@ -87,14 +60,14 @@ public:
 
 	/** Creates a profiler widget for a modifier profiler, returns default widget if none was registered */
 	TSharedPtr<SActorModifierCoreEditorProfiler> CreateProfilerWidget(TSharedPtr<FActorModifierCoreProfiler> InProfiler);
-	
+
 protected:
 	/** Modifier profiler pinned stats based on profiler type */
 	TMap<FName, TSet<FName>> ModifierProfilerStats;
 
 	/** Modifier profiler widgets */
 	TMap<FName, TFunction<TSharedRef<SActorModifierCoreEditorProfiler>(TSharedPtr<FActorModifierCoreProfiler>)>> ModifierProfilerWidgets;
-	
+
 	/** Runtime subsystem for modifier factories */
 	TWeakObjectPtr<UActorModifierCoreSubsystem> EngineSubsystem = nullptr;
 };
diff --git a/Engine/Plugins/Experimental/Avalanche/Source/AvalancheMaskEditor/Private/AvaMaskEditorMode.cpp b/Engine/Plugins/Experimental/Avalanche/Source/AvalancheMaskEditor/Private/AvaMaskEditorMode.cpp
--- a/Engine/Plugins/Experimental/Avalanche/Source/AvalancheMaskEditor/Private/AvaMaskEditorMode.cpp
+++ b/Engine/Plugins/Experimental/Avalanche/Source/AvalancheMaskEditor/Private/AvaMaskEditorMode.cpp
@@ -423,7 +423,7 @@ UAvaMask2DBaseModifier* UAvaMaskEditorMode::FindOrAddMaskModifier(AActor* InActo
 			FActorModifierCoreStackInsertOp InsertOp;
 			InsertOp.NewModifierName = MaskModifierName;
 			
-			MaskModifier = Cast<UAvaMask2DBaseModifier>(ModifierStack->InsertModifier(InsertOp));
+			MaskModifier = Cast<UAvaMask2DBaseModifier>(ModifierSubsystem->InsertModifier(ModifierStack, InsertOp));
 			if (!MaskModifier)
 			{
 				UE_LOG(LogAvalancheMaskEditor, Error, TEXT("Error inserting Mask modifier."));
@@ -446,16 +446,21 @@ bool UAvaMaskEditorMode::CanMaskSelected(AActor* InSelectedActor)
 	{
 		if (const UActorModifierCoreSubsystem* ModifierSubsystem = UActorModifierCoreSubsystem::Get())
 		{
-			const FName MaskModifierName = GetDefault<UAvaMask2DBaseModifier>()->GetModifierName();
-			
+			/*
+			 * Todo : this will always be none since UAvaMask2DBaseModifier is abstract and not registered as valid modifier,
+			 * should change this to UAvaMask2DReadModifier or UAvaMask2DWriteModifier based on context and avoid using abstract base class
+			 */
+
+			const FName MaskModifierName = ModifierSubsystem->GetRegisteredModifierName(UAvaMask2DBaseModifier::StaticClass());
+
 			if (const UActorModifierCoreStack* ExistingModifierStack = ModifierSubsystem->GetActorModifierStack(InActor))
 			{
-				if (ExistingModifierStack->ContainsModifier(MaskModifierName))
+				if (ExistingModifierStack->ContainsModifier(UAvaMask2DBaseModifier::StaticClass()))
 				{
 					return true;
-				}				
+				}
 			}
-			 
+ 
 			const TSet<FName> AllowedModifiers = ModifierSubsystem->GetAllowedModifiers(InActor);
 			if (AllowedModifiers.Contains(MaskModifierName))
 			{
diff --git a/Engine/Plugins/Experimental/Avalanche/Source/AvalancheModifiers/Private/Modifiers/AvaTaperModifier.cpp b/Engine/Plugins/Experimental/Avalanche/Source/AvalancheModifiers/Private/Modifiers/AvaTaperModifier.cpp
--- a/Engine/Plugins/Experimental/Avalanche/Source/AvalancheModifiers/Private/Modifiers/AvaTaperModifier.cpp
+++ b/Engine/Plugins/Experimental/Avalanche/Source/AvalancheModifiers/Private/Modifiers/AvaTaperModifier.cpp
@@ -270,7 +270,7 @@ int32 UAvaTaperModifier::GetSubdividersCuts() const
 {
 	int32 TotalSubdivisions = 0;
 	
-	if (GetModifierStack() && GetModifierStack()->ContainsModifier("Subdivide"))
+	if (GetModifierStack() && GetModifierStack()->ContainsModifier(UAvaSubdivideModifier::StaticClass()))
 	{
 		TArray<UAvaSubdivideModifier*> SubdivideModifiers;
 		GetModifierStack()->GetClassModifiers<UAvaSubdivideModifier>(SubdivideModifiers);
diff --git a/Engine/Plugins/Experimental/Avalanche/Source/AvalancheModifiers/Private/Tests/AvaModifiersAlignBetween.spec.cpp b/Engine/Plugins/Experimental/Avalanche/Source/AvalancheModifiers/Private/Tests/AvaModifiersAlignBetween.spec.cpp
--- a/Engine/Plugins/Experimental/Avalanche/Source/AvalancheModifiers/Private/Tests/AvaModifiersAlignBetween.spec.cpp
+++ b/Engine/Plugins/Experimental/Avalanche/Source/AvalancheModifiers/Private/Tests/AvaModifiersAlignBetween.spec.cpp
@@ -3,6 +3,7 @@
 #include "Misc/AutomationTest.h"
 #include "Modifiers/ActorModifierCoreStack.h"
 #include "Modifiers/AvaAlignBetweenModifier.h"
+#include "Subsystems/ActorModifierCoreSubsystem.h"
 #include "Tests/Framework/AvaModifiersTestUtils.h"
 #include "Tests/Framework/AvaTestDynamicMeshActor.h"
 #include "Tests/Framework/AvaTestUtils.h"
@@ -49,10 +50,11 @@ void AvalancheModifiersAlignBetween::Define()
 		                                        InitialReferenceTwoMeshActorLocation.Z * WeightIndex);
 
 		// Set up modifier
+		const UActorModifierCoreSubsystem* ModifierSubsystem = UActorModifierCoreSubsystem::Get();
 		const FName ModifierName = ModifierTestUtils->GetModifierName(UAvaAlignBetweenModifier::StaticClass());
 		FActorModifierCoreStackInsertOp InsertOp = ModifierTestUtils->GenerateInsertOp(ModifierName);
-		AlignBetweenModifier = Cast<UAvaAlignBetweenModifier>(
-			ModifierTestUtils->GenerateModifierStackForActor(ModifiedMeshActor)->InsertModifier(InsertOp));
+		UActorModifierCoreStack* ModifierStack = ModifierTestUtils->GenerateModifierStackForActor(ModifiedMeshActor);
+		AlignBetweenModifier = Cast<UAvaAlignBetweenModifier>(ModifierSubsystem->InsertModifier(ModifierStack, InsertOp));
 		check(AlignBetweenModifier);
 
 		// Set up the modifier
diff --git a/Engine/Plugins/Experimental/Avalanche/Source/AvalancheModifiers/Private/Tests/AvaModifiersAutoFollow.spec.cpp b/Engine/Plugins/Experimental/Avalanche/Source/AvalancheModifiers/Private/Tests/AvaModifiersAutoFollow.spec.cpp
--- a/Engine/Plugins/Experimental/Avalanche/Source/AvalancheModifiers/Private/Tests/AvaModifiersAutoFollow.spec.cpp
+++ b/Engine/Plugins/Experimental/Avalanche/Source/AvalancheModifiers/Private/Tests/AvaModifiersAutoFollow.spec.cpp
@@ -57,26 +57,22 @@ void AvalancheModifiersAutoFollowTests::Define()
 	{
 		BeforeEach([this]()
 		{
-			DynamicMeshActorToFollow = TestUtils->
-				SpawnTestDynamicMeshActor(FTransform(FVector(-100, -100, -100)));
+			DynamicMeshActorToFollow = TestUtils->SpawnTestDynamicMeshActor(FTransform(FVector(-100, -100, -100)));
 			FText OutFailReason = FText::GetEmpty();
 			FText* OutFailReasonPtr = &OutFailReason;
-			ActorModifierStack = ModifierSubsystem->
-				AddActorModifierStack(DynamicMeshActor);
+			ActorModifierStack = ModifierSubsystem->AddActorModifierStack(DynamicMeshActor);
 			check(ActorModifierStack);
 			FActorModifierCoreStackInsertOp InsertOp;
 
 			InsertOp.NewModifierName = ModifierAutoFollowName;
 			InsertOp.FailReason = OutFailReasonPtr;
-			AutoFollowModifier = Cast<UAvaAutoFollowModifier>(
-				ActorModifierStack->InsertModifier(InsertOp));
+			AutoFollowModifier = Cast<UAvaAutoFollowModifier>(ModifierSubsystem->InsertModifier(ActorModifierStack, InsertOp));
 		});
 
 		It("Should restore initial Dynamic mesh state once Autofollow modifier is disapplied", [this]()
 		{
 			FActorModifierCoreStackRemoveOp RemoveOp;
-			RemoveOp.RemoveModifier = AutoFollowModifier;
-			if (ActorModifierStack->RemoveModifier(RemoveOp))
+			if (ModifierSubsystem->RemoveModifiers({AutoFollowModifier}, RemoveOp))
 			{
 				FTransform RestoredActorState = DynamicMeshActor->GetActorTransform();
 				TestEqual("Actor returned to initial location", InitialActorState.GetLocation(),
@@ -99,7 +95,6 @@ void AvalancheModifiersAutoFollowTests::Define()
 				ReferenceActor.ReferenceContainer = EAvaReferenceContainer::Other;
 				ReferenceActor.ReferenceActorWeak = DynamicMeshActorToFollow;
 				AutoFollowModifier->SetReferenceActor(ReferenceActor);
-				const FActorModifierCoreMetadata& Metadata = AutoFollowModifier->GetModifierMetadata();
 				FTransform NewActorState = DynamicMeshActor->GetActorTransform();
 				TestNotEqual("Dynamic mesh position is changed", InitialActorState.GetLocation(),
 				             NewActorState.GetLocation());
diff --git a/Engine/Plugins/Experimental/Avalanche/Source/AvalancheModifiers/Private/Tests/AvaModifiersAutoSize.spec.cpp b/Engine/Plugins/Experimental/Avalanche/Source/AvalancheModifiers/Private/Tests/AvaModifiersAutoSize.spec.cpp
--- a/Engine/Plugins/Experimental/Avalanche/Source/AvalancheModifiers/Private/Tests/AvaModifiersAutoSize.spec.cpp
+++ b/Engine/Plugins/Experimental/Avalanche/Source/AvalancheModifiers/Private/Tests/AvaModifiersAutoSize.spec.cpp
@@ -4,6 +4,7 @@
 #include "Misc/AutomationTest.h"
 #include "Modifiers/ActorModifierCoreStack.h"
 #include "Modifiers/AvaAutoSizeModifier.h"
+#include "Subsystems/ActorModifierCoreSubsystem.h"
 #include "Tests/Framework/AvaModifiersTestUtils.h"
 #include "Tests/Framework/AvaTestDynamicMeshActor.h"
 #include "Tests/Framework/AvaTestUtils.h"
@@ -45,9 +46,11 @@ void AvalancheModifiersAutoSize::Define()
 		InitialModifiedActorBounds = FAvaModifiersActorUtils::GetActorBounds(ModifiedActor);
 
 		// Set up modifier stack
+		const UActorModifierCoreSubsystem* ModifierSubsystem = UActorModifierCoreSubsystem::Get();
 		const FName ModifierAutoSizeName = ModifierTestUtils->GetModifierName(UAvaAutoSizeModifier::StaticClass());
 		FActorModifierCoreStackInsertOp InsertOp = ModifierTestUtils->GenerateInsertOp(ModifierAutoSizeName);
-		AutoSizeModifier = Cast<UAvaAutoSizeModifier>(ModifierTestUtils->GenerateModifierStackForActor(ModifiedActor)->InsertModifier(InsertOp));
+		UActorModifierCoreStack* ModifierStack = ModifierTestUtils->GenerateModifierStackForActor(ModifiedActor);
+		AutoSizeModifier = Cast<UAvaAutoSizeModifier>(ModifierSubsystem->InsertModifier(ModifierStack, InsertOp));
 
 		// Set up the modifier
 		AutoSizeModifier->SetReferenceActor(ActorQuery);
diff --git a/Engine/Plugins/Experimental/Avalanche/Source/AvalancheModifiers/Private/Tests/AvaModifiersGridArrange.spec.cpp b/Engine/Plugins/Experimental/Avalanche/Source/AvalancheModifiers/Private/Tests/AvaModifiersGridArrange.spec.cpp
--- a/Engine/Plugins/Experimental/Avalanche/Source/AvalancheModifiers/Private/Tests/AvaModifiersGridArrange.spec.cpp
+++ b/Engine/Plugins/Experimental/Avalanche/Source/AvalancheModifiers/Private/Tests/AvaModifiersGridArrange.spec.cpp
@@ -3,6 +3,7 @@
 #include "Misc/AutomationTest.h"
 #include "Modifiers/ActorModifierCoreStack.h"
 #include "Modifiers/AvaGridArrangeModifier.h"
+#include "Subsystems/ActorModifierCoreSubsystem.h"
 #include "Tests/Framework/AvaModifiersTestUtils.h"
 #include "Tests/Framework/AvaTestDynamicMeshActor.h"
 #include "Tests/Framework/AvaTestUtils.h"
@@ -89,11 +90,11 @@ void AvalancheModifiersGridArrange::Define()
 		InitialChildFourActorState = ChildFourMeshActor->GetActorTransform();
 
 		// Set up modifier
-		const FName ModifierGridArrangeName = ModifierTestUtils->
-			GetModifierName(UAvaGridArrangeModifier::StaticClass());
+		const UActorModifierCoreSubsystem* ModifierSubsystem = UActorModifierCoreSubsystem::Get();
+		const FName ModifierGridArrangeName = ModifierTestUtils->GetModifierName(UAvaGridArrangeModifier::StaticClass());
 		FActorModifierCoreStackInsertOp InsertOp = ModifierTestUtils->GenerateInsertOp(ModifierGridArrangeName);
-		GridArrangeModifier = Cast<UAvaGridArrangeModifier>(
-			ModifierTestUtils->GenerateModifierStackForActor(ParentMeshActor)->InsertModifier(InsertOp));
+		UActorModifierCoreStack* ModifierStack = ModifierTestUtils->GenerateModifierStackForActor(ParentMeshActor);
+		GridArrangeModifier = Cast<UAvaGridArrangeModifier>(ModifierSubsystem->InsertModifier(ModifierStack, InsertOp));
 	});
 
 	AfterEach([this]
diff --git a/Engine/Plugins/Experimental/Avalanche/Source/AvalancheModifiers/Private/Tests/AvaModifiersJustify.spec.cpp b/Engine/Plugins/Experimental/Avalanche/Source/AvalancheModifiers/Private/Tests/AvaModifiersJustify.spec.cpp
--- a/Engine/Plugins/Experimental/Avalanche/Source/AvalancheModifiers/Private/Tests/AvaModifiersJustify.spec.cpp
+++ b/Engine/Plugins/Experimental/Avalanche/Source/AvalancheModifiers/Private/Tests/AvaModifiersJustify.spec.cpp
@@ -3,6 +3,7 @@
 #include "Misc/AutomationTest.h"
 #include "Modifiers/ActorModifierCoreStack.h"
 #include "Modifiers/AvaJustifyModifier.h"
+#include "Subsystems/ActorModifierCoreSubsystem.h"
 #include "Tests/Framework/AvaModifiersTestUtils.h"
 #include "Tests/Framework/AvaTestDynamicMeshActor.h"
 #include "Tests/Framework/AvaTestUtils.h"
@@ -80,10 +81,11 @@ void AvalancheModifiersJustify::Define()
 		InitialChildThreeLocation = ChildThreeActor->GetActorTransform().GetLocation();
 
 		// Set up modifier
+		const UActorModifierCoreSubsystem* ModifierSubsystem = UActorModifierCoreSubsystem::Get();
 		const FName ModifierJustifyName = ModifierTestUtils->GetModifierName(UAvaJustifyModifier::StaticClass());
 		FActorModifierCoreStackInsertOp InsertOp = ModifierTestUtils->GenerateInsertOp(ModifierJustifyName);
-		JustifyModifier = Cast<UAvaJustifyModifier>(
-			ModifierTestUtils->GenerateModifierStackForActor(ParentActor)->InsertModifier(InsertOp));
+		UActorModifierCoreStack* ModifierStack = ModifierTestUtils->GenerateModifierStackForActor(ParentActor);
+		JustifyModifier = Cast<UAvaJustifyModifier>(ModifierSubsystem->InsertModifier(ModifierStack, InsertOp));
 		JustifyModifier->SetHorizontalAlignment(EAvaJustifyHorizontal::Right);
 		JustifyModifier->SetVerticalAlignment(EAvaJustifyVertical::Bottom);
 		JustifyModifier->SetDepthAlignment(EAvaJustifyDepth::Center);
diff --git a/Engine/Plugins/Experimental/Avalanche/Source/AvalancheModifiers/Private/Tests/AvaModifiersLookAt.spec.cpp b/Engine/Plugins/Experimental/Avalanche/Source/AvalancheModifiers/Private/Tests/AvaModifiersLookAt.spec.cpp
--- a/Engine/Plugins/Experimental/Avalanche/Source/AvalancheModifiers/Private/Tests/AvaModifiersLookAt.spec.cpp
+++ b/Engine/Plugins/Experimental/Avalanche/Source/AvalancheModifiers/Private/Tests/AvaModifiersLookAt.spec.cpp
@@ -4,6 +4,7 @@
 #include "Misc/AutomationTest.h"
 #include "Modifiers/ActorModifierCoreStack.h"
 #include "Modifiers/AvaLookAtModifier.h"
+#include "Subsystems/ActorModifierCoreSubsystem.h"
 #include "Tests/Framework/AvaModifiersTestUtils.h"
 #include "Tests/Framework/AvaTestDynamicMeshActor.h"
 #include "Tests/Framework/AvaTestUtils.h"
@@ -47,10 +48,11 @@ void AvalancheModifiersLookAt::Define()
 		InitialReferenceSate = ReferenceActor->GetActorTransform();
 
 		// Set up modifier
+		const UActorModifierCoreSubsystem* ModifierSubsystem = UActorModifierCoreSubsystem::Get();
 		const FName ModifierLookAtName = ModifierTestUtils->GetModifierName(UAvaLookAtModifier::StaticClass());
 		FActorModifierCoreStackInsertOp InsertOp = ModifierTestUtils->GenerateInsertOp(ModifierLookAtName);
-		LookAtModifier = Cast<UAvaLookAtModifier>(
-			ModifierTestUtils->GenerateModifierStackForActor(ModifiedActor)->InsertModifier(InsertOp));
+		UActorModifierCoreStack* ModifierStack = ModifierTestUtils->GenerateModifierStackForActor(ModifiedActor);
+		LookAtModifier = Cast<UAvaLookAtModifier>(ModifierSubsystem->InsertModifier(ModifierStack, InsertOp));
 		LookAtModifier->SetReferenceActor(FAvaSceneTreeActor(ReferenceActor));
 	});
 
diff --git a/Engine/Plugins/Experimental/Avalanche/Source/AvalancheModifiers/Private/Tests/AvaModifiersRadialArrange.cpp b/Engine/Plugins/Experimental/Avalanche/Source/AvalancheModifiers/Private/Tests/AvaModifiersRadialArrange.cpp
--- a/Engine/Plugins/Experimental/Avalanche/Source/AvalancheModifiers/Private/Tests/AvaModifiersRadialArrange.cpp
+++ b/Engine/Plugins/Experimental/Avalanche/Source/AvalancheModifiers/Private/Tests/AvaModifiersRadialArrange.cpp
@@ -3,6 +3,7 @@
 #include "Misc/AutomationTest.h"
 #include "Modifiers/ActorModifierCoreStack.h"
 #include "Modifiers/AvaRadialArrangeModifier.h"
+#include "Subsystems/ActorModifierCoreSubsystem.h"
 #include "Tests/Framework/AvaModifiersTestUtils.h"
 #include "Tests/Framework/AvaTestDynamicMeshActor.h"
 #include "Tests/Framework/AvaTestUtils.h"
@@ -37,11 +38,11 @@ void AvalancheModifiersRadialArrange::Define()
 		InitialParentLocation = ParentActor->GetActorLocation();
 
 		// Set up modifier
+		const UActorModifierCoreSubsystem* ModifierSubsystem = UActorModifierCoreSubsystem::Get();
 		const FName ModifierName = ModifierTestUtils->GetModifierName(UAvaRadialArrangeModifier::StaticClass());
 		FActorModifierCoreStackInsertOp InsertOp = ModifierTestUtils->GenerateInsertOp(ModifierName);
-		RadialArrangeModifier = Cast<UAvaRadialArrangeModifier>(
-			ModifierTestUtils->GenerateModifierStackForActor(ParentActor)->InsertModifier(InsertOp)
-		);
+		UActorModifierCoreStack* ModifierStack = ModifierTestUtils->GenerateModifierStackForActor(ParentActor);
+		RadialArrangeModifier = Cast<UAvaRadialArrangeModifier>(ModifierSubsystem->InsertModifier(ModifierStack, InsertOp));
 		InnerRadius = FMath::RandRange(-1000, 1000);
 		RadialArrangeModifier->SetInnerRadius(InnerRadius);
 	});
diff --git a/Engine/Plugins/Experimental/Avalanche/Source/AvalancheModifiers/Private/Tests/AvaModifiersTranslucentPriority.spec.cpp b/Engine/Plugins/Experimental/Avalanche/Source/AvalancheModifiers/Private/Tests/AvaModifiersTranslucentPriority.spec.cpp
--- a/Engine/Plugins/Experimental/Avalanche/Source/AvalancheModifiers/Private/Tests/AvaModifiersTranslucentPriority.spec.cpp
+++ b/Engine/Plugins/Experimental/Avalanche/Source/AvalancheModifiers/Private/Tests/AvaModifiersTranslucentPriority.spec.cpp
@@ -4,6 +4,7 @@
 #include "Components/PrimitiveComponent.h"
 #include "Modifiers/ActorModifierCoreStack.h"
 #include "Modifiers/AvaTranslucentPriorityModifier.h"
+#include "Subsystems/ActorModifierCoreSubsystem.h"
 #include "Tests/Framework/AvaModifiersTestUtils.h"
 #include "Tests/Framework/AvaTestDynamicMeshActor.h"
 #include "Tests/Framework/AvaTestUtils.h"
@@ -36,11 +37,11 @@ void AvalancheModifiersTranslucentPriority::Define()
 
 		ModifiedActor->GetComponents(ModifiedActorComponents);
 
+		const UActorModifierCoreSubsystem* ModifierSubsystem = UActorModifierCoreSubsystem::Get();
 		const FName ModifierName = ModifierTestUtils->GetModifierName(UAvaTranslucentPriorityModifier::StaticClass());
 		FActorModifierCoreStackInsertOp InsertOp = ModifierTestUtils->GenerateInsertOp(ModifierName);
-		TranslucentPriorityModifier = Cast<UAvaTranslucentPriorityModifier>(
-			ModifierTestUtils->GenerateModifierStackForActor(ModifiedActor)->InsertModifier(InsertOp)
-		);
+		UActorModifierCoreStack* ModifierStack = ModifierTestUtils->GenerateModifierStackForActor(ModifiedActor);
+		TranslucentPriorityModifier = Cast<UAvaTranslucentPriorityModifier>(ModifierSubsystem->InsertModifier(ModifierStack, InsertOp));
 		TranslucentPriorityModifier->SetMode(EAvaTranslucentPriorityModifierMode::Manual);
 		TranslucentPriorityModifier->SetSortPriority(ModifierSortPriority);
 	});
diff --git a/Engine/Plugins/Experimental/Avalanche/Source/AvalancheModifiers/Private/Tests/AvaModifiersVisibility.spec.cpp b/Engine/Plugins/Experimental/Avalanche/Source/AvalancheModifiers/Private/Tests/AvaModifiersVisibility.spec.cpp
--- a/Engine/Plugins/Experimental/Avalanche/Source/AvalancheModifiers/Private/Tests/AvaModifiersVisibility.spec.cpp
+++ b/Engine/Plugins/Experimental/Avalanche/Source/AvalancheModifiers/Private/Tests/AvaModifiersVisibility.spec.cpp
@@ -3,6 +3,7 @@
 #include "Misc/AutomationTest.h"
 #include "Modifiers/ActorModifierCoreStack.h"
 #include "Modifiers/AvaVisibilityModifier.h"
+#include "Subsystems/ActorModifierCoreSubsystem.h"
 #include "Tests/Framework/AvaModifiersTestUtils.h"
 #include "Tests/Framework/AvaTestDynamicMeshActor.h"
 #include "Tests/Framework/AvaTestUtils.h"
@@ -45,11 +46,11 @@ void AvalancheModifiersVisibility::Define()
 		ChildrenActors = ModifierTestUtils->SpawnTestDynamicMeshActors(NumberOfChildren, ParentActor);
 
 		// Set up modifier
+		const UActorModifierCoreSubsystem* ModifierSubsystem = UActorModifierCoreSubsystem::Get();
 		const FName ModifierName = ModifierTestUtils->GetModifierName(UAvaVisibilityModifier::StaticClass());
 		FActorModifierCoreStackInsertOp InsertOp = ModifierTestUtils->GenerateInsertOp(ModifierName);
-		VisibilityModifier = Cast<UAvaVisibilityModifier>(
-			ModifierTestUtils->GenerateModifierStackForActor(ParentActor)->InsertModifier(InsertOp)
-		);
+		UActorModifierCoreStack* ModifierStack = ModifierTestUtils->GenerateModifierStackForActor(ParentActor);
+		VisibilityModifier = Cast<UAvaVisibilityModifier>(ModifierSubsystem->InsertModifier(ModifierStack, InsertOp));
 		VisibilityModifier->SetIndex(NumberOfChildrenVisible - 1);
 	});
 
diff --git a/Engine/Plugins/Experimental/Avalanche/Source/AvalancheModifiersEditor/Private/AvaModifiersEditorStyle.cpp b/Engine/Plugins/Experimental/Avalanche/Source/AvalancheModifiersEditor/Private/AvaModifiersEditorStyle.cpp
--- a/Engine/Plugins/Experimental/Avalanche/Source/AvalancheModifiersEditor/Private/AvaModifiersEditorStyle.cpp
+++ b/Engine/Plugins/Experimental/Avalanche/Source/AvalancheModifiersEditor/Private/AvaModifiersEditorStyle.cpp
@@ -38,7 +38,7 @@ FAvaModifiersEditorStyle::FAvaModifiersEditorStyle()
 	if (const UActorModifierCoreSubsystem* ModifierSubsystem = UActorModifierCoreSubsystem::Get())
 	{
 		// loop through already registered factories if subsystem exists
-		ModifierSubsystem->ForEachModifierMetadata([this](FActorModifierCoreMetadata& InMetadata)->bool
+		ModifierSubsystem->ForEachModifierMetadata([this](const FActorModifierCoreMetadata& InMetadata)->bool
 		{
 			OnModifierClassRegistered(InMetadata);
 			return true;
diff --git a/Engine/Plugins/Experimental/Avalanche/Source/AvalancheModifiersEditor/Private/Outliner/AvaOutlinerModifier.cpp b/Engine/Plugins/Experimental/Avalanche/Source/AvalancheModifiersEditor/Private/Outliner/AvaOutlinerModifier.cpp
--- a/Engine/Plugins/Experimental/Avalanche/Source/AvalancheModifiersEditor/Private/Outliner/AvaOutlinerModifier.cpp
+++ b/Engine/Plugins/Experimental/Avalanche/Source/AvalancheModifiersEditor/Private/Outliner/AvaOutlinerModifier.cpp
@@ -78,7 +78,10 @@ void FAvaOutlinerModifier::OnVisibilityChanged(EAvaOutlinerVisibilityType InVisi
 {
 	if (InVisibilityType == EAvaOutlinerVisibilityType::Runtime && Modifier.IsValid())
 	{
-		Modifier->SetModifierEnabled(bInNewVisibility);
+		if (const UActorModifierCoreSubsystem* ModifierSubsystem = UActorModifierCoreSubsystem::Get())
+		{
+			ModifierSubsystem->EnableModifiers({Modifier.Get()}, bInNewVisibility, true);
+		}
 	}
 }
 
diff --git a/Engine/Plugins/Experimental/Avalanche/Source/AvalancheModifiersEditor/Private/Outliner/AvaOutlinerModifierDropHandler.cpp b/Engine/Plugins/Experimental/Avalanche/Source/AvalancheModifiersEditor/Private/Outliner/AvaOutlinerModifierDropHandler.cpp
--- a/Engine/Plugins/Experimental/Avalanche/Source/AvalancheModifiersEditor/Private/Outliner/AvaOutlinerModifierDropHandler.cpp
+++ b/Engine/Plugins/Experimental/Avalanche/Source/AvalancheModifiersEditor/Private/Outliner/AvaOutlinerModifierDropHandler.cpp
@@ -46,19 +46,19 @@ TOptional<EItemDropZone> FAvaOutlinerModifierDropHandler::CanDrop(EItemDropZone
 	}
 
 	UActorModifierCoreBase* TargetModifier = TargetModifierItem->GetModifier();
-	const UActorModifierCoreEditorSubsystem* Extension = UActorModifierCoreEditorSubsystem::Get();
+	const UActorModifierCoreSubsystem* ModifierSubsystem = UActorModifierCoreSubsystem::Get();
 	const EActorModifierCoreStackPosition Position = InDropZone == EItemDropZone::AboveItem ? EActorModifierCoreStackPosition::Before : EActorModifierCoreStackPosition::After;
-	
+
 	TArray<UActorModifierCoreBase*> MoveModifiers;
 	TArray<UActorModifierCoreBase*> CloneModifiers;
-	Extension->GetSortedModifiers(GetDraggedModifiers(), TargetModifier->GetModifiedActor(), TargetModifier, Position, MoveModifiers, CloneModifiers);
-	
-	// If the Modifiers are Empty, return fail early 
+	ModifierSubsystem->GetSortedModifiers(GetDraggedModifiers(), TargetModifier->GetModifiedActor(), TargetModifier, Position, MoveModifiers, CloneModifiers);
+
+	// If the Modifiers are Empty, return fail early
 	if (MoveModifiers.IsEmpty() && CloneModifiers.IsEmpty())
 	{
 		return TOptional<EItemDropZone>();
 	}
-	
+
 	return InDropZone;
 }
 
@@ -154,19 +154,17 @@ TOptional<EItemDropZone> FAvaOutlinerModifierDropHandler::CanDropOnActor(AActor*
 bool FAvaOutlinerModifierDropHandler::DropModifiersInActor(AActor* InActor, EItemDropZone InDropZone) const
 {
 	const UActorModifierCoreSubsystem* const ModifierSubsystem = UActorModifierCoreSubsystem::Get();
-	const UActorModifierCoreEditorSubsystem* const ModifierExtensionSubsystem = UActorModifierCoreEditorSubsystem::Get();
 
-	if (!IsValid(ModifierSubsystem) || !IsValid(ModifierExtensionSubsystem))
+	if (!IsValid(ModifierSubsystem))
 	{
 		return false;
 	}
 
 	const EActorModifierCoreStackPosition Position = InDropZone == EItemDropZone::AboveItem ? EActorModifierCoreStackPosition::Before : EActorModifierCoreStackPosition::After;
-	uint32 AddedModifierCount = 0;
 
 	TArray<UActorModifierCoreBase*> MoveModifiers;
 	TArray<UActorModifierCoreBase*> CloneModifiers;
-	ModifierExtensionSubsystem->GetSortedModifiers(GetDraggedModifiers(), InActor, nullptr, Position, MoveModifiers, CloneModifiers);
+	ModifierSubsystem->GetSortedModifiers(GetDraggedModifiers(), InActor, nullptr, Position, MoveModifiers, CloneModifiers);
 
 	if (CloneModifiers.IsEmpty())
 	{
@@ -177,31 +175,24 @@ bool FAvaOutlinerModifierDropHandler::DropModifiersInActor(AActor* InActor, EIte
 		, CloneModifiers.Num()
 		, *InActor->GetActorNameOrLabel());
 
+	bool bSuccess = false;
+
 	if (UActorModifierCoreStack* const Stack = ModifierSubsystem->AddActorModifierStack(InActor))
 	{
 		FText FailReason;
-		
-		Stack->ProcessLockFunction([&CloneModifiers, &FailReason, Stack, &AddedModifierCount, InActor]()
+		FActorModifierCoreStackCloneOp CloneOp;
+		CloneOp.bShouldTransact = true;
+		CloneOp.FailReason = &FailReason;
+
+		bSuccess = ModifierSubsystem->CloneModifiers(CloneModifiers, Stack, CloneOp).Num() == CloneModifiers.Num();
+
+		if (!bSuccess)
 		{
-			for (UActorModifierCoreBase* const DropModifier : CloneModifiers)
-			{
-				FActorModifierCoreStackCloneOp CloneOp;
-				CloneOp.FailReason = &FailReason;
-				CloneOp.CloneModifier = DropModifier;
-				if (Stack->CloneModifier(CloneOp))
-				{
-					++AddedModifierCount; 
-				}
-				else
-				{
-					UE_LOG(LogAvaOutlinerModifierDropHandler, Warning, TEXT("Clone modifier %s on actor %s failed : %s"),
-						*CloneOp.CloneModifier->GetModifierName().ToString(),
-						*InActor->GetActorNameOrLabel(),
-						*FailReason.ToString());
-					break;
-				}
-			}
-		});
+			UE_LOG(LogAvaOutlinerModifierDropHandler, Warning, TEXT("Clone %i modifier(s) on actor %s failed : %s"),
+				CloneModifiers.Num(),
+				*InActor->GetActorNameOrLabel(),
+				*FailReason.ToString());
+		}
 
 		if (!FailReason.IsEmpty())
 		{
@@ -212,17 +203,16 @@ bool FAvaOutlinerModifierDropHandler::DropModifiersInActor(AActor* InActor, EIte
 		}
 	}
 
-	return AddedModifierCount > 0;
+	return bSuccess;
 }
 
 bool FAvaOutlinerModifierDropHandler::DropModifiersInModifier(UActorModifierCoreBase* InTargetModifier, EItemDropZone InDropZone) const
 {
 	const UActorModifierCoreSubsystem* const ModifierSubsystem = UActorModifierCoreSubsystem::Get();
-	const UActorModifierCoreEditorSubsystem* const ExtensionSubsystem = UActorModifierCoreEditorSubsystem::Get();
 
 	const UActorModifierCoreStack* const TargetStack = InTargetModifier->GetModifierStack();
 
-	if (!IsValid(ModifierSubsystem) || !IsValid(ExtensionSubsystem) || !IsValid(TargetStack))
+	if (!IsValid(ModifierSubsystem) || !IsValid(TargetStack))
 	{
 		return false;
 	}
@@ -232,11 +222,11 @@ bool FAvaOutlinerModifierDropHandler::DropModifiersInModifier(UActorModifierCore
 
 	TArray<UActorModifierCoreBase*> MoveModifiers;
 	TArray<UActorModifierCoreBase*> CloneModifiers;
-	ExtensionSubsystem->GetSortedModifiers(GetDraggedModifiers(), TargetActor, InTargetModifier, Position, MoveModifiers, CloneModifiers);
+	ModifierSubsystem->GetSortedModifiers(GetDraggedModifiers(), TargetActor, InTargetModifier, Position, MoveModifiers, CloneModifiers);
 
 	bool bSuccess = false;
 	FText FailReason;
-	
+
 	if (!MoveModifiers.IsEmpty())
 	{
 		UE_LOG(LogAvaOutlinerModifierDropHandler, Log, TEXT("Dropping %i modifier(s) %s modifier %s on actor %s (Move)")
@@ -245,7 +235,13 @@ bool FAvaOutlinerModifierDropHandler::DropModifiersInModifier(UActorModifierCore
 		, *InTargetModifier->GetModifierName().ToString()
 		, *TargetActor->GetActorNameOrLabel());
 
-		bSuccess = ExtensionSubsystem->MoveModifiers(MoveModifiers, InTargetModifier, Position, &FailReason, true);
+		FActorModifierCoreStackMoveOp MoveOp;
+		MoveOp.bShouldTransact = true;
+		MoveOp.FailReason = &FailReason;
+		MoveOp.MovePosition = Position;
+		MoveOp.MovePositionContext = InTargetModifier;
+
+		bSuccess = ModifierSubsystem->MoveModifiers(MoveModifiers, InTargetModifier->GetModifierStack(), MoveOp);
 	}
 	else if (!CloneModifiers.IsEmpty())
 	{
@@ -255,9 +251,15 @@ bool FAvaOutlinerModifierDropHandler::DropModifiersInModifier(UActorModifierCore
 		, *InTargetModifier->GetModifierName().ToString()
 		, *TargetActor->GetActorNameOrLabel());
 
-		bSuccess = ExtensionSubsystem->CloneModifiers(CloneModifiers, InTargetModifier, Position, &FailReason, true);
+		FActorModifierCoreStackCloneOp CloneOp;
+		CloneOp.bShouldTransact = true;
+		CloneOp.FailReason = &FailReason;
+		CloneOp.ClonePosition = Position;
+		CloneOp.ClonePositionContext = InTargetModifier;
+
+		bSuccess = ModifierSubsystem->CloneModifiers(CloneModifiers, InTargetModifier->GetModifierStack(), CloneOp).Num() == CloneModifiers.Num();
 	}
-	
+
 	if (!FailReason.IsEmpty())
 	{
 		FNotificationInfo NotificationInfo(FailReason);
@@ -265,7 +267,7 @@ bool FAvaOutlinerModifierDropHandler::DropModifiersInModifier(UActorModifierCore
 		NotificationInfo.bFireAndForget = true;
 		FSlateNotificationManager::Get().AddNotification(NotificationInfo);
 	}
-	
+
 	return bSuccess;
 }
 
diff --git a/Engine/Plugins/Experimental/Avalanche/Source/AvalancheSVGEditor/Private/AvaSVGEditorModule.cpp b/Engine/Plugins/Experimental/Avalanche/Source/AvalancheSVGEditor/Private/AvaSVGEditorModule.cpp
--- a/Engine/Plugins/Experimental/Avalanche/Source/AvalancheSVGEditor/Private/AvaSVGEditorModule.cpp
+++ b/Engine/Plugins/Experimental/Avalanche/Source/AvalancheSVGEditor/Private/AvaSVGEditorModule.cpp
@@ -69,16 +69,15 @@ void FAvaSVGEditorModule::OnSVGActorSplit(ASVGShapesParentActor* InSVGShapesPare
 				}
 
 				FActorModifierCoreStackInsertOp ExtrudeModifierInsertOp;
-				ExtrudeModifierInsertOp.NewModifierName = GetDefault<UAvaExtrudeModifier>()->GetModifierName();
+				ExtrudeModifierInsertOp.NewModifierName = ModifierCoreSubsystem->GetRegisteredModifierName(UAvaExtrudeModifier::StaticClass());
 
 				if (!ModifierCoreSubsystem->GetAllowedModifiers(ShapeActor).Contains(ExtrudeModifierInsertOp.NewModifierName))
 				{
 					continue;
 				}
 
-				UActorModifierCoreBase* Modifier = nullptr;
+				UActorModifierCoreBase* Modifier = ModifierCoreSubsystem->InsertModifier(ModifierStack, ExtrudeModifierInsertOp);
 
-				Modifier = ModifierStack->InsertModifier(ExtrudeModifierInsertOp);
 				if (UAvaExtrudeModifier* ExtrudeModifier = Cast<UAvaExtrudeModifier>(Modifier))
 				{
 					EAvaExtrudeMode ExtrudeMode = EAvaExtrudeMode::Symmetrical;
@@ -109,14 +108,15 @@ void FAvaSVGEditorModule::OnSVGActorSplit(ASVGShapesParentActor* InSVGShapesPare
 				if (Shape->Bevel > 0.0f)
 				{
 					FActorModifierCoreStackInsertOp BevelModifierInsertOp;
-					BevelModifierInsertOp.NewModifierName = GetDefault<UAvaBevelModifier>()->GetModifierName();
+					BevelModifierInsertOp.NewModifierName = ModifierCoreSubsystem->GetRegisteredModifierName(UAvaBevelModifier::StaticClass());
 
 					if (!ModifierCoreSubsystem->GetAllowedModifiers(ShapeActor).Contains(BevelModifierInsertOp.NewModifierName))
 					{
 						continue;
 					}
 
-					Modifier = ModifierStack->InsertModifier(BevelModifierInsertOp);
+					Modifier = ModifierCoreSubsystem->InsertModifier(ModifierStack, BevelModifierInsertOp);
+
 					if (UAvaBevelModifier* BevelModifier = Cast<UAvaBevelModifier>(Modifier))
 					{
 						BevelModifier->SetInset(Shape->Bevel);
