Actor Modifiers:
- Fixed issue with Shared Actor getting spawned while iterating array on stream level unload.

#changelist validated
#virtualized

diff --git a/Engine/Plugins/Experimental/ActorModifierCore/Source/ActorModifierCore/Private/Subsystems/ActorModifierCoreSubsystem.cpp b/Engine/Plugins/Experimental/ActorModifierCore/Source/ActorModifierCore/Private/Subsystems/ActorModifierCoreSubsystem.cpp
--- a/Engine/Plugins/Experimental/ActorModifierCore/Source/ActorModifierCore/Private/Subsystems/ActorModifierCoreSubsystem.cpp
+++ b/Engine/Plugins/Experimental/ActorModifierCore/Source/ActorModifierCore/Private/Subsystems/ActorModifierCoreSubsystem.cpp
@@ -886,7 +886,7 @@ UActorModifierCoreSharedObject* UActorModifierCoreSubsystem::GetModifierSharedOb
 		return nullptr;
 	}
 
-	if (AActorModifierCoreSharedActor* SharedActor = GetModifierSharedProvider(InLevel))
+	if (AActorModifierCoreSharedActor* SharedActor = GetModifierSharedProvider(InLevel, bInCreateIfNone))
 	{
 		if (bInCreateIfNone)
 		{
