Motion Design:
Fixed issue with normal Levels not differentiated from 'Motion Design' levels in Usage Statistics. To solve these, analytics were added to Motion Design Editor Created/Activated/Deactivated.

In addition to this the following analytics were added:
- Interactive Tool Completed (e.g. Spawn Cloner)
- Material Designer Material Created
- Ava Sequence Added
- Modifiers Created
- Property Animator Created

#changelist validated
#virtualized

diff --git a/Engine/Plugins/Experimental/ActorModifierCore/Source/ActorModifierCore/Private/Subsystems/ActorModifierCoreSubsystem.cpp b/Engine/Plugins/Experimental/ActorModifierCore/Source/ActorModifierCore/Private/Subsystems/ActorModifierCoreSubsystem.cpp
--- a/Engine/Plugins/Experimental/ActorModifierCore/Source/ActorModifierCore/Private/Subsystems/ActorModifierCoreSubsystem.cpp
+++ b/Engine/Plugins/Experimental/ActorModifierCore/Source/ActorModifierCore/Private/Subsystems/ActorModifierCoreSubsystem.cpp
@@ -1,9 +1,8 @@
 // Copyright Epic Games, Inc. All Rights Reserved.
 
 #include "Subsystems/ActorModifierCoreSubsystem.h"
-
-#include "EngineUtils.h"
 #include "Engine/Engine.h"
+#include "EngineUtils.h"
 #include "GameFramework/Actor.h"
 #include "Internationalization/Text.h"
 #include "Modifiers/ActorModifierCoreComponent.h"
@@ -13,6 +12,7 @@
 #include "UObject/UObjectIterator.h"
 
 #if WITH_EDITOR
+#include "EngineAnalytics.h"
 #include "ScopedTransaction.h"
 #endif
 
@@ -59,6 +59,18 @@ void UActorModifierCoreSubsystem::Deinitialize()
 	Super::Deinitialize();
 }
 
+void UActorModifierCoreSubsystem::OnInsertModifier(const FActorModifierCoreStackInsertOp& InInsertOp) const
+{
+#if WITH_EDITOR
+	if (FEngineAnalytics::IsAvailable())
+	{
+		TArray<FAnalyticsEventAttribute> Attributes;
+		Attributes.Emplace(TEXT("Name"), InInsertOp.NewModifierName.ToString());
+		FEngineAnalytics::GetProvider().RecordEvent(TEXT("Editor.Usage.ActorModifiers.InsertModifier"), Attributes);
+	}
+#endif
+}
+
 bool UActorModifierCoreSubsystem::RegisterModifierClass(const UClass* InModifierClass, bool bInOverrideIfExists)
 {
 	if (!IsValid(InModifierClass))
@@ -1155,6 +1167,8 @@ TArray<UActorModifierCoreBase*> UActorModifierCoreSubsystem::AddActorsModifiers(
 	InAddOp.InsertPositionContext = nullptr;
 	InAddOp.InsertPosition = EActorModifierCoreStackPosition::Before;
 
+	OnInsertModifier(InAddOp);
+
 	for (UActorModifierCoreStack* ActorStack : ActorStacks)
 	{
 		NewModifiers.Add(ActorStack->InsertModifier(InAddOp));
@@ -1215,6 +1229,8 @@ UActorModifierCoreBase* UActorModifierCoreSubsystem::InsertModifier(UActorModifi
 	FScopedTransaction Transaction(TransactionText, InInsertOp.bShouldTransact);
 #endif
 
+	OnInsertModifier(InInsertOp);
+
 	// insert modifier in stack
 	NewModifier = InStack->InsertModifier(InInsertOp);
 
diff --git a/Engine/Plugins/Experimental/ActorModifierCore/Source/ActorModifierCore/Public/Subsystems/ActorModifierCoreSubsystem.h b/Engine/Plugins/Experimental/ActorModifierCore/Source/ActorModifierCore/Public/Subsystems/ActorModifierCoreSubsystem.h
--- a/Engine/Plugins/Experimental/ActorModifierCore/Source/ActorModifierCore/Public/Subsystems/ActorModifierCoreSubsystem.h
+++ b/Engine/Plugins/Experimental/ActorModifierCore/Source/ActorModifierCore/Public/Subsystems/ActorModifierCoreSubsystem.h
@@ -131,6 +131,8 @@ protected:
 	virtual void Deinitialize() override;
 	//~ End UEngineSubsystem
 
+	void OnInsertModifier(const FActorModifierCoreStackInsertOp& InInsertOp) const;
+
 	/** Register a root actor stack to query this actor stack from everywhere, is called automatically at creation or deserialization by stack */
 	bool RegisterActorModifierStack(UActorModifierCoreStack* InStack);
 
diff --git a/Engine/Plugins/Experimental/Avalanche/Source/Avalanche/Private/AvaScene.cpp b/Engine/Plugins/Experimental/Avalanche/Source/Avalanche/Private/AvaScene.cpp
--- a/Engine/Plugins/Experimental/Avalanche/Source/Avalanche/Private/AvaScene.cpp
+++ b/Engine/Plugins/Experimental/Avalanche/Source/Avalanche/Private/AvaScene.cpp
@@ -17,6 +17,7 @@
 
 #if WITH_EDITOR
 #include "AvaField.h"
+#include "EngineAnalytics.h"
 #include "Misc/ScopedSlowTask.h"
 #include "RemoteControlBinding.h"
 #endif
@@ -25,6 +26,16 @@ DEFINE_LOG_CATEGORY_STATIC(LogAvaScene, Log, All);
 
 #define LOCTEXT_NAMESPACE "AvaScene"
 
+void AAvaScene::OnSceneCreated(FString&& InCreationType)
+{
+#if WITH_EDITOR
+	if (FEngineAnalytics::IsAvailable())
+	{
+		FEngineAnalytics::GetProvider().RecordEvent(TEXT("Editor.Usage.MotionDesign.SceneCreated"), FAnalyticsEventAttribute(TEXT("CreationType"), MoveTemp(InCreationType)));
+	}
+#endif
+}
+
 AAvaScene* AAvaScene::GetScene(ULevel* InLevel, bool bInCreateSceneIfNotFound)
 {
 	if (!IsValid(InLevel))
@@ -53,6 +64,7 @@ AAvaScene* AAvaScene::GetScene(ULevel* InLevel, bool bInCreateSceneIfNotFound)
 #endif
 
 	AAvaScene* const NewScene = World->SpawnActor<AAvaScene>(SpawnParameters);
+	OnSceneCreated(/*CreationType*/TEXT("Spawned"));
 	return NewScene;
 }
 
@@ -170,6 +182,14 @@ bool AAvaScene::AddSequence(UAvaSequence* InSequence)
 	{
 		Animations.Add(InSequence);
 		ScheduleRebuildSequenceTree();
+
+#if WITH_EDITOR
+		if (FEngineAnalytics::IsAvailable())
+		{
+			FEngineAnalytics::GetProvider().RecordEvent(TEXT("Editor.Usage.MotionDesign.AddedSequence"));
+		}
+#endif
+
 		return true;
 	}
 	return false;
diff --git a/Engine/Plugins/Experimental/Avalanche/Source/Avalanche/Public/AvaScene.h b/Engine/Plugins/Experimental/Avalanche/Source/Avalanche/Public/AvaScene.h
--- a/Engine/Plugins/Experimental/Avalanche/Source/Avalanche/Public/AvaScene.h
+++ b/Engine/Plugins/Experimental/Avalanche/Source/Avalanche/Public/AvaScene.h
@@ -27,6 +27,8 @@ class AAvaScene : public AActor, public IAvaSequenceProvider, public IAvaSceneIn
 {
 	GENERATED_BODY()
 
+	static void OnSceneCreated(FString&& InCreationType);
+
 public:
 	AVALANCHE_API static AAvaScene* GetScene(ULevel* InLevel, bool bInCreateSceneIfNotFound);
 
diff --git a/Engine/Plugins/Experimental/Avalanche/Source/AvalancheEditorCore/Private/AvaEditor.cpp b/Engine/Plugins/Experimental/Avalanche/Source/AvalancheEditorCore/Private/AvaEditor.cpp
--- a/Engine/Plugins/Experimental/Avalanche/Source/AvalancheEditorCore/Private/AvaEditor.cpp
+++ b/Engine/Plugins/Experimental/Avalanche/Source/AvalancheEditorCore/Private/AvaEditor.cpp
@@ -12,6 +12,7 @@
 #include "Elements/Framework/TypedElementRegistry.h"
 #include "Elements/Framework/TypedElementUtil.h"
 #include "Elements/Interfaces/TypedElementObjectInterface.h"
+#include "EngineAnalytics.h"
 #include "Framework/Docking/LayoutExtender.h"
 #include "Framework/Docking/WorkspaceItem.h"
 #include "HAL/PlatformApplicationMisc.h"
@@ -223,6 +224,23 @@ void FAvaEditor::ForEachExtension(TFunctionRef<void(const TSharedRef<IAvaEditorE
 	}
 }
 
+void FAvaEditor::RecordActivationChangedEvent()
+{
+	if (!FEngineAnalytics::IsAvailable())
+	{
+		return;
+	}
+
+	if (bIsActive)
+	{
+		FEngineAnalytics::GetProvider().RecordEvent(TEXT("Editor.Usage.MotionDesign.Activated"));
+	}
+	else
+	{
+		FEngineAnalytics::GetProvider().RecordEvent(TEXT("Editor.Usage.MotionDesign.Deactivated"));
+	}
+}
+
 void FAvaEditor::Activate(TSharedPtr<IToolkitHost> InOverrideToolkitHost)
 {
 	if (!CanActivate())
@@ -234,6 +252,8 @@ void FAvaEditor::Activate(TSharedPtr<IToolkitHost> InOverrideToolkitHost)
 
 	bIsActive = true;
 
+	RecordActivationChangedEvent();
+
 	BindDelegates();
 
 	if (InOverrideToolkitHost.IsValid())
@@ -286,6 +306,8 @@ void FAvaEditor::Deactivate()
 
 	bIsActive = false;
 
+	RecordActivationChangedEvent();
+
 	UnbindDelegates();
 
 	DeactivateExtensions();
diff --git a/Engine/Plugins/Experimental/Avalanche/Source/AvalancheEditorCore/Private/AvaEditor.h b/Engine/Plugins/Experimental/Avalanche/Source/AvalancheEditorCore/Private/AvaEditor.h
--- a/Engine/Plugins/Experimental/Avalanche/Source/AvalancheEditorCore/Private/AvaEditor.h
+++ b/Engine/Plugins/Experimental/Avalanche/Source/AvalancheEditorCore/Private/AvaEditor.h
@@ -35,6 +35,8 @@ protected:
 
 	void ForEachExtension(TFunctionRef<void(const TSharedRef<IAvaEditorExtension>&)> InFunc) const;
 
+	void RecordActivationChangedEvent();
+
 	//~ Begin IAvaEditor
 	virtual void Activate(TSharedPtr<IToolkitHost> InOverrideToolkitHost = nullptr) override;
 	virtual void Deactivate() override;
diff --git a/Engine/Plugins/Experimental/Avalanche/Source/AvalancheInteractiveTools/Private/Tools/AvaInteractiveToolsToolBase.cpp b/Engine/Plugins/Experimental/Avalanche/Source/AvalancheInteractiveTools/Private/Tools/AvaInteractiveToolsToolBase.cpp
--- a/Engine/Plugins/Experimental/Avalanche/Source/AvalancheInteractiveTools/Private/Tools/AvaInteractiveToolsToolBase.cpp
+++ b/Engine/Plugins/Experimental/Avalanche/Source/AvalancheInteractiveTools/Private/Tools/AvaInteractiveToolsToolBase.cpp
@@ -2,15 +2,16 @@
 
 #include "Tools/AvaInteractiveToolsToolBase.h"
 #include "AvaInteractiveToolsSettings.h"
-#include "AvalancheInteractiveToolsModule.h"
 #include "AvaViewportUtils.h"
+#include "AvalancheInteractiveToolsModule.h"
 #include "BaseBehaviors/SingleClickBehavior.h"
 #include "BaseBehaviors/SingleKeyCaptureBehavior.h"
 #include "ContextObjectStore.h"
+#include "EdMode/AvaInteractiveToolsEdMode.h"
 #include "Editor/EditorEngine.h"
 #include "EditorViewportClient.h"
-#include "EdMode/AvaInteractiveToolsEdMode.h"
 #include "Engine/World.h"
+#include "EngineAnalytics.h"
 #include "Framework/Application/SlateApplication.h"
 #include "IAvaInteractiveToolsModeDetailsObject.h"
 #include "IAvaInteractiveToolsModeDetailsObjectProvider.h"
@@ -573,6 +574,17 @@ void UAvaInteractiveToolsToolBase::OnComplete()
 		PreviewActor->Destroy();
 		PreviewActor = nullptr;
 	}
+
+	if (FEngineAnalytics::IsAvailable())
+	{
+		TArray<FAnalyticsEventAttribute> Attributes;
+		Attributes.Emplace(TEXT("ToolClass"), GetClass()->GetName());
+		if (SpawnedActor)
+		{
+			Attributes.Emplace(TEXT("ActorClass"), SpawnedActor->GetClass()->GetName());	
+		}
+		FEngineAnalytics::GetProvider().RecordEvent(TEXT("Editor.Usage.MotionDesign.CompleteTool"), Attributes);
+	}
 }
 
 void UAvaInteractiveToolsToolBase::BeginTransaction()
diff --git a/Engine/Plugins/Experimental/DynamicMaterial/Source/DynamicMaterialEditor/Private/Material/DynamicMaterialInstanceFactory.cpp b/Engine/Plugins/Experimental/DynamicMaterial/Source/DynamicMaterialEditor/Private/Material/DynamicMaterialInstanceFactory.cpp
--- a/Engine/Plugins/Experimental/DynamicMaterial/Source/DynamicMaterialEditor/Private/Material/DynamicMaterialInstanceFactory.cpp
+++ b/Engine/Plugins/Experimental/DynamicMaterial/Source/DynamicMaterialEditor/Private/Material/DynamicMaterialInstanceFactory.cpp
@@ -3,11 +3,12 @@
 #include "Material/DynamicMaterialInstanceFactory.h"
 #include "AssetToolsModule.h"
 #include "DMPrivate.h"
+#include "EngineAnalytics.h"
 #include "IAssetTools.h"
 #include "Material/DynamicMaterialInstance.h"
 #include "Model/DynamicMaterialModel.h"
-#include "Model/IDynamicMaterialModelEditorOnlyDataInterface.h"
 #include "Model/DynamicMaterialModelFactory.h"
+#include "Model/IDynamicMaterialModelEditorOnlyDataInterface.h"
 #include "Modules/ModuleManager.h"
 
 #define LOCTEXT_NAMESPACE "MaterialDesignerInstanceFactory"
@@ -30,6 +31,11 @@ UObject* UDynamicMaterialInstanceFactory::FactoryCreateNew(UClass* Class, UObjec
 	UDynamicMaterialInstance* NewInstance = NewObject<UDynamicMaterialInstance>(InParent, Class, Name, Flags | RF_Transactional);
 	check(NewInstance);
 
+	if (FEngineAnalytics::IsAvailable())
+	{
+		FEngineAnalytics::GetProvider().RecordEvent(TEXT("Editor.Usage.MaterialDesigner.CreateMaterial"));
+	}
+
 	UDynamicMaterialModelFactory* EditorFactory = NewObject<UDynamicMaterialModelFactory>();
 	check(EditorFactory);
 
diff --git a/Engine/Plugins/Experimental/PropertyAnimatorCore/Source/PropertyAnimatorCore/Private/Subsystems/PropertyAnimatorCoreSubsystem.cpp b/Engine/Plugins/Experimental/PropertyAnimatorCore/Source/PropertyAnimatorCore/Private/Subsystems/PropertyAnimatorCoreSubsystem.cpp
--- a/Engine/Plugins/Experimental/PropertyAnimatorCore/Source/PropertyAnimatorCore/Private/Subsystems/PropertyAnimatorCoreSubsystem.cpp
+++ b/Engine/Plugins/Experimental/PropertyAnimatorCore/Source/PropertyAnimatorCore/Private/Subsystems/PropertyAnimatorCoreSubsystem.cpp
@@ -1,21 +1,21 @@
 // Copyright Epic Games, Inc. All Rights Reserved.
 
 #include "Subsystems/PropertyAnimatorCoreSubsystem.h"
-
 #include "Animators/PropertyAnimatorCoreBase.h"
 #include "Components/PropertyAnimatorCoreComponent.h"
 #include "Engine/Engine.h"
 #include "GameFramework/Actor.h"
 #include "Presets/PropertyAnimatorCorePresetBase.h"
 #include "Properties/Converters/PropertyAnimatorCoreConverterBase.h"
-#include "Properties/PropertyAnimatorCoreResolver.h"
 #include "Properties/Handlers/PropertyAnimatorCoreHandlerBase.h"
+#include "Properties/PropertyAnimatorCoreResolver.h"
 #include "TimeSources/PropertyAnimatorCoreTimeSourceBase.h"
 #include "TimeSources/PropertyAnimatorCoreWorldTimeSource.h"
 #include "UObject/Class.h"
 #include "UObject/UObjectIterator.h"
 
 #if WITH_EDITOR
+#include "EngineAnalytics.h"
 #include "ScopedTransaction.h"
 #endif
 
@@ -323,6 +323,13 @@ TSet<UPropertyAnimatorCoreBase*> UPropertyAnimatorCoreSubsystem::CreateAnimators
 	const FText ActorCount = FText::FromString(FString::FromInt(InActors.Num()));
 
 	FScopedTransaction Transaction(FText::Format(TransactionText, AnimatorName, ActorCount), bInTransact);
+
+	if (FEngineAnalytics::IsAvailable())
+	{
+		TArray<FAnalyticsEventAttribute> Attributes;
+		Attributes.Emplace(TEXT("Class"), InAnimatorClass->GetName());
+		FEngineAnalytics::GetProvider().RecordEvent(TEXT("Editor.Usage.PropertyAnimator.CreateAnimator"), Attributes);
+	}
 #endif
 
 	for (AActor* Actor : InActors)
