MotionDesign : Modifiers
- Only log modifier optimization warning when we invalidate all during Clone/Insert/Move/Remove operations

#changelist validated
#virtualized

diff --git a/Engine/Plugins/Experimental/ActorModifierCore/Source/ActorModifierCore/Private/Modifiers/ActorModifierCoreStack.cpp b/Engine/Plugins/Experimental/ActorModifierCore/Source/ActorModifierCore/Private/Modifiers/ActorModifierCoreStack.cpp
--- a/Engine/Plugins/Experimental/ActorModifierCore/Source/ActorModifierCore/Private/Modifiers/ActorModifierCoreStack.cpp
+++ b/Engine/Plugins/Experimental/ActorModifierCore/Source/ActorModifierCore/Private/Modifiers/ActorModifierCoreStack.cpp
@@ -637,7 +637,7 @@ UActorModifierCoreBase* UActorModifierCoreStack::CloneModifier(FActorModifierCor
 		}
 	}
 
-	LogModifier(FString::Printf(TEXT("Insert modifier %s"), *InCloneOp.CloneModifier->GetModifierName().ToString()), true);
+	LogModifier(FString::Printf(TEXT("Clone modifier %s"), *InCloneOp.CloneModifier->GetModifierName().ToString()), true);
 
 	// Refresh the stack
 	MarkModifierDirty();
@@ -1525,7 +1525,12 @@ void UActorModifierCoreStack::CheckModifierOptimization(bool bInInvalidateAll)
 					const FText Message = FText::Format(LOCTEXT("AvoidAfterCategory", "Should be moved above {0} modifiers"), FText::FromName(Category));
 					Modifier->Status = FActorModifierCoreStatus(EActorModifierCoreStatus::Warning, Message);
 
-					LogModifier(FString::Printf(TEXT("Optimisation possible for modifier %s : %s"), *Modifier->GetModifierName().ToString(), *Message.ToString()), true);
+					// Only log when we invalidate all and recheck optimization
+					if (bInInvalidateAll)
+					{
+						LogModifier(FString::Printf(TEXT("Optimisation possible for modifier %s : %s"), *Modifier->GetModifierName().ToString(), *Message.ToString()), true);
+					}
+
 					break;
 				}
 			}
@@ -1550,7 +1555,12 @@ void UActorModifierCoreStack::CheckModifierOptimization(bool bInInvalidateAll)
 					const FText Message = FText::Format(LOCTEXT("AvoidBeforeCategory", "Should be moved below {0} modifiers"), FText::FromName(Category));
 					Modifier->Status = FActorModifierCoreStatus(EActorModifierCoreStatus::Warning, Message);
 
-					LogModifier(FString::Printf(TEXT("Optimisation possible for modifier %s : %s"), *Modifier->GetModifierName().ToString(), *Message.ToString()), true);
+					// Only log when we invalidate all and recheck optimization
+					if (bInInvalidateAll)
+					{
+						LogModifier(FString::Printf(TEXT("Optimisation possible for modifier %s : %s"), *Modifier->GetModifierName().ToString(), *Message.ToString()), true);
+					}
+
 					break;
 				}
 			}
