MotionDesign : Modifiers
- Fix init issue caused by world and actor still in loading phase and affecting modifier apply operation

diff --git a/Engine/Plugins/Experimental/ActorModifierCore/Source/ActorModifierCore/Private/Modifiers/ActorModifierCoreBase.cpp b/Engine/Plugins/Experimental/ActorModifierCore/Source/ActorModifierCore/Private/Modifiers/ActorModifierCoreBase.cpp
--- a/Engine/Plugins/Experimental/ActorModifierCore/Source/ActorModifierCore/Private/Modifiers/ActorModifierCoreBase.cpp
+++ b/Engine/Plugins/Experimental/ActorModifierCore/Source/ActorModifierCore/Private/Modifiers/ActorModifierCoreBase.cpp
@@ -599,8 +599,10 @@ void UActorModifierCoreBase::PostModifierWorldLoad(UWorld* InWorld, ELevelTick I
 {
 	const AActor* Actor = GetModifiedActor();
 
+	// Check actor is in the world loaded and does not need post load and is not in async loading
 	if (Actor
 		&& InWorld == Actor->GetWorld()
+		&& !Actor->HasAnyInternalFlags(EInternalObjectFlags::AsyncLoading)
 		&& !Actor->HasAnyFlags(EObjectFlags::RF_NeedPostLoad)
 		&& !Actor->HasAnyFlags(EObjectFlags::RF_NeedPostLoadSubobjects))
 	{
