MotionDesign : ActorModifierCore
- Fix issue with spawning actors with modifiers via template actors (sequencer)

diff --git a/Engine/Plugins/Experimental/ActorModifierCore/Source/ActorModifierCore/Private/Modifiers/ActorModifierCoreBase.cpp b/Engine/Plugins/Experimental/ActorModifierCore/Source/ActorModifierCore/Private/Modifiers/ActorModifierCoreBase.cpp
--- a/Engine/Plugins/Experimental/ActorModifierCore/Source/ActorModifierCore/Private/Modifiers/ActorModifierCoreBase.cpp
+++ b/Engine/Plugins/Experimental/ActorModifierCore/Source/ActorModifierCore/Private/Modifiers/ActorModifierCoreBase.cpp
@@ -2,9 +2,11 @@
 
 #include "Modifiers/ActorModifierCoreBase.h"
 
+#include "Modifiers/ActorModifierCoreComponent.h"
 #include "Modifiers/ActorModifierCoreSharedObject.h"
 #include "Modifiers/ActorModifierCoreStack.h"
 #include "Subsystems/ActorModifierCoreSubsystem.h"
+#include "UObject/Package.h"
 
 DEFINE_LOG_CATEGORY_STATIC(LogActorModifierCoreBase, Log, All);
 
@@ -155,7 +157,7 @@ void UActorModifierCoreBase::Unapply()
 
 void UActorModifierCoreBase::OnModifierDirty(UActorModifierCoreBase* DirtyModifier, bool bExecute)
 {
-	if (ModifierStack.IsValid())
+	if (UActorModifierCoreStack* ModifierStack = GetModifierStack())
 	{
 		ModifierStack->OnModifierDirty(DirtyModifier, bExecute);
 	}
@@ -185,11 +187,16 @@ AActor* UActorModifierCoreBase::GetModifiedActor() const
 {
 	if (!ModifiedActor.IsValid())
 	{
-		const_cast<UActorModifierCoreBase*>(this)->ModifiedActor = Cast<AActor>(GetOuter());
+		const_cast<UActorModifierCoreBase*>(this)->ModifiedActor = GetTypedOuter<AActor>();
 	}
 	return ModifiedActor.Get();
 }
 
+UActorModifierCoreStack* UActorModifierCoreBase::GetModifierStack() const
+{
+	return GetTypedOuter<UActorModifierCoreStack>();
+}
+
 UActorModifierCoreStack* UActorModifierCoreBase::GetRootModifierStack() const
 {
 	// we are not the root stack
@@ -420,6 +427,20 @@ bool UActorModifierCoreBase::ProcessFunction(TFunctionRef<bool(const UActorModif
 	return InFunction(this);
 }
 
+void UActorModifierCoreBase::DeferInitializeModifier()
+{
+	// Begin batch operation to avoid updating every time a modifier is loaded
+	UActorModifierCoreStack* Stack = GetRootModifierStack();
+	if (!Stack->IsModifierExecutionLocked() && !Stack->IsModifierStackInitialized())
+	{
+		Stack->LockModifierExecution();
+	}
+
+	// Bind to world delegate, tick will be called when all actors have been loaded, unbind when actors have been loaded
+	FWorldDelegates::OnWorldPostActorTick.RemoveAll(this);
+	FWorldDelegates::OnWorldPostActorTick.AddUObject(this, &UActorModifierCoreBase::PostModifierWorldLoad);
+}
+
 UActorModifierCoreBase::UActorModifierCoreBase()
 {
 	// Copy metadata from CDO
@@ -434,16 +455,45 @@ void UActorModifierCoreBase::PostLoad()
 {
 	Super::PostLoad();
 
-	// Begin batch operation to avoid updating every time a modifier is loaded
-	UActorModifierCoreStack* Stack = GetRootModifierStack();
-	if (!Stack->IsModifierExecutionLocked() && !Stack->IsModifierStackInitialized())
+	const AActor* OwningActor = GetModifiedActor();
+	const UObject* Outer = GetOuter();
+	const UActorModifierCoreComponent* OwningComponent = OwningActor->FindComponentByClass<UActorModifierCoreComponent>();
+
+	if (OwningActor
+		&& OwningComponent
+		&& Outer
+		&& Outer->IsA<AActor>())
 	{
-		Stack->LockModifierExecution();
+		constexpr int32 RenameFlags = REN_DontCreateRedirectors | REN_DoNotDirty | REN_ForceNoResetLoaders | REN_NonTransactional;
+
+		// Migrate modifiers to component stack instead and discard this stack
+		if (const UActorModifierCoreStack* ThisStack = Cast<UActorModifierCoreStack>(this))
+		{
+			UActorModifierCoreStack* ComponentStack = OwningComponent->ModifierStack;
+
+			if (ComponentStack && ComponentStack != this)
+			{
+				const FString ThisStackName = GetName();
+				const EObjectFlags ThisStackFlags = GetFlags();
+
+				LogModifier(FString::Printf(TEXT("Modifier stack migrated to component stack %s %s"), *OwningComponent->GetName(), *ComponentStack->GetName()), true);
+				Rename(TEXT("Trash_ModifierStack"), GetTransientPackage(), RenameFlags);
+
+				ComponentStack->Rename(*ThisStackName, nullptr, RenameFlags);
+				ComponentStack->Modifiers = ThisStack->Modifiers;
+				ComponentStack->bModifierProfiling = ThisStack->bModifierProfiling;
+				ComponentStack->SetFlags(ThisStackFlags);
+			}
+		}
+		// Change outer of modifier to component stack instead of actor
+		else
+		{
+			const bool bSuccess = Rename(nullptr, OwningComponent->GetModifierStack(), RenameFlags);
+			LogModifier(FString::Printf(TEXT("Modifier outer renamed to stack %s : %s"), *OwningComponent->GetModifierStack()->GetName(), bSuccess ? TEXT("OK") : TEXT("Fail")), true);
+		}
 	}
 
-	// Bind to world delegate, tick will be called when all actors have been loaded, unbind when actors have been loaded
-	FWorldDelegates::OnWorldPostActorTick.RemoveAll(this);
-	FWorldDelegates::OnWorldPostActorTick.AddUObject(this, &UActorModifierCoreBase::PostModifierWorldLoad);
+	DeferInitializeModifier();
 }
 
 void UActorModifierCoreBase::PostEditImport()
@@ -486,9 +536,11 @@ void UActorModifierCoreBase::PostEditUndo()
 {
 	Super::PostEditUndo();
 
+	UActorModifierCoreStack* ModifierStack = GetModifierStack();
+
 	// is it an undo remove or undo add operation ?
-	const bool bModifierInStack = ModifierStack.IsValid() && ModifierStack->Modifiers.Contains(this);
-	const bool bStackRegistered = IsModifierStack() && !ModifierStack.IsValid() && ModifiedActor.IsValid();
+	const bool bModifierInStack = ModifierStack && ModifierStack->Modifiers.Contains(this);
+	const bool bStackRegistered = IsModifierStack() && !ModifierStack && ModifiedActor.IsValid();
 	const bool bModifierValid = bModifierInStack || bStackRegistered;
 
 	if (!bModifierValid)
@@ -550,7 +602,7 @@ void UActorModifierCoreBase::LogModifier(const FString& InLog, bool bInForce) co
 {
 	if (IsModifierProfiling() || bInForce)
 	{
-		const FString ActorLabel = ModifiedActor.IsValid() ? *GetModifiedActor()->GetActorNameOrLabel() : TEXT("Invalid actor");
+		const FString ActorLabel = GetModifiedActor() ? *GetModifiedActor()->GetActorNameOrLabel() : TEXT("Invalid actor");
 		const FString ModifierLabel = GetModifierName().ToString();
 
 		UE_LOG(LogActorModifierCoreBase, Log, TEXT("[%s][%s] %s"), *ActorLabel, *ModifierLabel, *InLog);
@@ -573,14 +625,13 @@ void UActorModifierCoreBase::TickModifier(float InDeltaTime)
 void UActorModifierCoreBase::PostModifierCreation(UActorModifierCoreStack* InStack)
 {
 	// initialize once, called by the subsystem itself
-	if (!ModifierStack.IsValid())
+	if (GetModifierStack() == InStack)
 	{
 		if (const UActorModifierCoreBase* CDO = GetClass()->GetDefaultObject<UActorModifierCoreBase>())
 		{
 			Metadata = CDO->Metadata;
 		}
 
-		ModifierStack = InStack;
 		ModifiedActor = GetModifiedActor();
 		bModifierInitialized = false;
 	}
@@ -628,7 +679,7 @@ void UActorModifierCoreBase::PostModifierWorldLoad(UWorld* InWorld, ELevelTick I
 		InitializeModifier(EActorModifierCoreEnableReason::Load);
 
 		// End batch operation and execute all modifiers at once if all stack is initialized
-		if (IsModifierStack() && !ModifierStack.IsValid())
+		if (GetRootModifierStack() == this)
 		{
 			if (IsModifierExecutionLocked())
 			{
@@ -686,7 +737,7 @@ void UActorModifierCoreBase::InitializeModifier(EActorModifierCoreEnableReason I
 #endif
 
 		// set new actor
-		ModifiedActor = Cast<AActor>(GetOuter());
+		ModifiedActor = GetModifiedActor();
 
 		// Initialize profiler
 		if (!Profiler.IsValid())
@@ -736,7 +787,7 @@ void UActorModifierCoreBase::UninitializeModifier(EActorModifierCoreDisableReaso
 		OnModifierRemoved(InReason);
 
 		// set new actor
-		ModifiedActor = Cast<AActor>(GetOuter());
+		ModifiedActor = GetModifiedActor();
 
 		// recover old enabled state
 		bModifierEnabled = bWasModifierEnabled;
diff --git a/Engine/Plugins/Experimental/ActorModifierCore/Source/ActorModifierCore/Private/Modifiers/ActorModifierCoreComponent.cpp b/Engine/Plugins/Experimental/ActorModifierCore/Source/ActorModifierCore/Private/Modifiers/ActorModifierCoreComponent.cpp
--- a/Engine/Plugins/Experimental/ActorModifierCore/Source/ActorModifierCore/Private/Modifiers/ActorModifierCoreComponent.cpp
+++ b/Engine/Plugins/Experimental/ActorModifierCore/Source/ActorModifierCore/Private/Modifiers/ActorModifierCoreComponent.cpp
@@ -5,6 +5,9 @@
 
 UActorModifierCoreComponent::UActorModifierCoreComponent()
 {
+	ModifierStack = CreateDefaultSubobject<UActorModifierCoreStack>(TEXT("ModifierStack"));
+	ModifierStack->PostModifierCreation(/** Parentstack */nullptr);
+
 	if (!IsTemplate())
 	{
 		PrimaryComponentTick.bCanEverTick = true;
@@ -52,7 +55,11 @@ UActorModifierCoreComponent* UActorModifierCoreComponent::CreateAndExposeCompone
 void UActorModifierCoreComponent::OnComponentCreated()
 {
 	Super::OnComponentCreated();
-	InitializeStack();
+
+	if (ModifierStack && !ModifierStack->IsModifierInitialized())
+	{
+		ModifierStack->InitializeModifier(EActorModifierCoreEnableReason::User);
+	}
 }
 
 void UActorModifierCoreComponent::OnComponentDestroyed(bool bDestroyingHierarchy)
@@ -82,22 +89,6 @@ void UActorModifierCoreComponent::PostEditUndo()
 }
 #endif
 
-void UActorModifierCoreComponent::InitializeStack()
-{
-	if (!ModifierStack)
-	{
-		if (AActor* OwningActor = Cast<AActor>(GetOuter()))
-		{
-			ModifierStack = UActorModifierCoreStack::Create(OwningActor, nullptr);
-
-			if (ModifierStack)
-			{
-				ModifierStack->SetModifierEnabled(true);
-			}
-		}
-	}
-}
-
 void UActorModifierCoreComponent::TickComponent(float InDeltaTime, ELevelTick InTickType, FActorComponentTickFunction* InThisTickFunction)
 {
 	Super::TickComponent(InDeltaTime, InTickType, InThisTickFunction);
diff --git a/Engine/Plugins/Experimental/ActorModifierCore/Source/ActorModifierCore/Private/Modifiers/ActorModifierCoreDefs.cpp b/Engine/Plugins/Experimental/ActorModifierCore/Source/ActorModifierCore/Private/Modifiers/ActorModifierCoreDefs.cpp
--- a/Engine/Plugins/Experimental/ActorModifierCore/Source/ActorModifierCore/Private/Modifiers/ActorModifierCoreDefs.cpp
+++ b/Engine/Plugins/Experimental/ActorModifierCore/Source/ActorModifierCore/Private/Modifiers/ActorModifierCoreDefs.cpp
@@ -250,7 +250,7 @@ UActorModifierCoreBase* FActorModifierCoreMetadata::CreateModifierInstance(UActo
 {
 	if (InStack && InStack->GetModifiedActor())
 	{
-		UActorModifierCoreBase* NewModifierInstance = NewObject<UActorModifierCoreBase>(InStack->GetModifiedActor(), Class, NAME_None, RF_Transactional);
+		UActorModifierCoreBase* NewModifierInstance = NewObject<UActorModifierCoreBase>(InStack, Class, NAME_None, RF_Transactional);
 		NewModifierInstance->PostModifierCreation(InStack);
 		return NewModifierInstance;
 	}
diff --git a/Engine/Plugins/Experimental/ActorModifierCore/Source/ActorModifierCore/Private/Modifiers/ActorModifierCoreStack.cpp b/Engine/Plugins/Experimental/ActorModifierCore/Source/ActorModifierCore/Private/Modifiers/ActorModifierCoreStack.cpp
--- a/Engine/Plugins/Experimental/ActorModifierCore/Source/ActorModifierCore/Private/Modifiers/ActorModifierCoreStack.cpp
+++ b/Engine/Plugins/Experimental/ActorModifierCore/Source/ActorModifierCore/Private/Modifiers/ActorModifierCoreStack.cpp
@@ -4,6 +4,7 @@
 
 #include "Algo/RemoveIf.h"
 #include "Components/SceneComponent.h"
+#include "Modifiers/ActorModifierCoreComponent.h"
 #include "Subsystems/ActorModifierCoreSubsystem.h"
 
 #define LOCTEXT_NAMESPACE "ActorModifierCoreStack"
@@ -14,11 +15,20 @@ UActorModifierCoreStack::FOnModifierUpdated UActorModifierCoreStack::OnModifierA
 UActorModifierCoreStack::FOnModifierUpdated UActorModifierCoreStack::OnModifierRemovedDelegate;
 UActorModifierCoreStack::FOnModifierUpdated UActorModifierCoreStack::OnModifierMovedDelegate;
 
-UActorModifierCoreStack* UActorModifierCoreStack::Create(AActor* InActor, UActorModifierCoreStack* InParentStack)
+UActorModifierCoreStack* UActorModifierCoreStack::Create(UActorModifierCoreComponent* InComponent, UActorModifierCoreStack* InParentStack)
 {
-	if (IsValid(InActor) && (!InParentStack || InParentStack->GetModifiedActor() == InActor))
+	if (IsValid(InComponent) && (!InParentStack || InParentStack->GetModifiedActor() == InComponent->GetOwner()))
 	{
-		UActorModifierCoreStack* NewStack = NewObject<UActorModifierCoreStack>(InActor, NAME_None, RF_Transactional);
+		UObject* NewOuter = nullptr;
+		if (InParentStack)
+		{
+			NewOuter = InParentStack;
+		}
+		else
+		{
+			NewOuter = InComponent;
+		}
+		UActorModifierCoreStack* NewStack = NewObject<UActorModifierCoreStack>(NewOuter, NAME_None, RF_Transactional);
 		NewStack->PostModifierCreation(InParentStack);
 		NewStack->InitializeModifier(EActorModifierCoreEnableReason::User);
 		return NewStack;
@@ -344,6 +354,11 @@ TArray<UActorModifierCoreBase*> UActorModifierCoreStack::FindModifiers(const UCl
 	return FoundModifiers;
 }
 
+bool UActorModifierCoreStack::IsRootStack() const
+{
+	return !GetModifierStack();
+}
+
 bool UActorModifierCoreStack::ContainsModifierBefore(const FName& InSearchName, const UActorModifierCoreBase* InBeforeModifier) const
 {
 	if (!IsValid(InBeforeModifier))
@@ -594,9 +609,8 @@ UActorModifierCoreBase* UActorModifierCoreStack::CloneModifier(FActorModifierCor
 		UActorModifierCoreBase* NewModifier = nullptr;
 		if (RequireModifier == InCloneOp.CloneModifier->GetModifierName())
 		{
-			FObjectDuplicationParameters Parameters = InitStaticDuplicateObjectParams(InCloneOp.CloneModifier, GetModifiedActor());
+			FObjectDuplicationParameters Parameters = InitStaticDuplicateObjectParams(InCloneOp.CloneModifier, this);
 			NewModifier = Cast<UActorModifierCoreBase>(StaticDuplicateObjectEx(Parameters));
-			NewModifier->ModifierStack = this;
 			NewModifier->ModifiedActor = GetModifiedActor();
 		}
 		else
diff --git a/Engine/Plugins/Experimental/ActorModifierCore/Source/ActorModifierCore/Public/Modifiers/ActorModifierCoreBase.h b/Engine/Plugins/Experimental/ActorModifierCore/Source/ActorModifierCore/Public/Modifiers/ActorModifierCoreBase.h
--- a/Engine/Plugins/Experimental/ActorModifierCore/Source/ActorModifierCore/Public/Modifiers/ActorModifierCoreBase.h
+++ b/Engine/Plugins/Experimental/ActorModifierCore/Source/ActorModifierCore/Public/Modifiers/ActorModifierCoreBase.h
@@ -65,10 +65,7 @@ public:
 	ACTORMODIFIERCORE_API AActor* GetModifiedActor() const;
 
 	/** Returns the stack this modifier is in */
-	UActorModifierCoreStack* GetModifierStack() const
-	{
-		return ModifierStack.Get();
-	}
+	ACTORMODIFIERCORE_API UActorModifierCoreStack* GetModifierStack() const;
 
 	/** Returns the top root stack this modifier is in */
 	ACTORMODIFIERCORE_API UActorModifierCoreStack* GetRootModifierStack() const;
@@ -352,6 +349,8 @@ protected:
 	/** Logs modifier message if in profiling mode or forced */
 	void LogModifier(const FString& InLog, bool bInForce = false) const;
 
+	void DeferInitializeModifier();
+
 private:
 	/** Called when modifier becomes dirty */
 	ACTORMODIFIERCORE_API virtual void OnModifierDirty(UActorModifierCoreBase* DirtyModifier, bool bExecute);
@@ -421,12 +420,8 @@ private:
 	/** Promise that executes when the modifier execution is done */
 	TSharedPtr<TPromise<bool>, ESPMode::ThreadSafe> ExecutePromise = nullptr;
 
-	UPROPERTY(DuplicateTransient, NonTransactional)
 	TWeakObjectPtr<AActor> ModifiedActor = nullptr;
 
-	UPROPERTY(NonTransactional)
-	TWeakObjectPtr<UActorModifierCoreStack> ModifierStack = nullptr;
-
 	/** Is the modifier enabled or disabled */
 	UPROPERTY(EditInstanceOnly, Setter="SetModifierEnabled", Getter="IsModifierEnabled", Category="Modifier", meta=(DisplayName="Enable Modifier"))
 	bool bModifierEnabled = true;
diff --git a/Engine/Plugins/Experimental/ActorModifierCore/Source/ActorModifierCore/Public/Modifiers/ActorModifierCoreComponent.h b/Engine/Plugins/Experimental/ActorModifierCore/Source/ActorModifierCore/Public/Modifiers/ActorModifierCoreComponent.h
--- a/Engine/Plugins/Experimental/ActorModifierCore/Source/ActorModifierCore/Public/Modifiers/ActorModifierCoreComponent.h
+++ b/Engine/Plugins/Experimental/ActorModifierCore/Source/ActorModifierCore/Public/Modifiers/ActorModifierCoreComponent.h
@@ -13,6 +13,8 @@ class UActorModifierCoreComponent : public UActorComponent
 {
 	GENERATED_BODY()
 
+	friend class UActorModifierCoreBase;
+
 public:
 	/** Create this component for a specific actor and exposes it to the editor and for serialization */
 	static UActorModifierCoreComponent* CreateAndExposeComponent(AActor* InParentActor);
@@ -36,12 +38,9 @@ protected:
 #endif
 	//~ End UObject
 
-	/** Initialize the inner root stack for the actor it is attached */
-	void InitializeStack();
-
 	virtual void TickComponent(float InDeltaTime, ELevelTick InTickType, FActorComponentTickFunction* InThisTickFunction) override;
 
 	/** This is the root stack that will contain all modifiers for this actor */
-	UPROPERTY(VisibleAnywhere, BlueprintReadOnly, NoClear, Export, Instanced, Category="Modifiers")
+	UPROPERTY(VisibleAnywhere, BlueprintReadOnly, Category="Modifiers")
 	TObjectPtr<UActorModifierCoreStack> ModifierStack = nullptr;
 };
\ No newline at end of file
diff --git a/Engine/Plugins/Experimental/ActorModifierCore/Source/ActorModifierCore/Public/Modifiers/ActorModifierCoreStack.h b/Engine/Plugins/Experimental/ActorModifierCore/Source/ActorModifierCore/Public/Modifiers/ActorModifierCoreStack.h
--- a/Engine/Plugins/Experimental/ActorModifierCore/Source/ActorModifierCore/Public/Modifiers/ActorModifierCoreStack.h
+++ b/Engine/Plugins/Experimental/ActorModifierCore/Source/ActorModifierCore/Public/Modifiers/ActorModifierCoreStack.h
@@ -7,6 +7,7 @@
 #include "ActorModifierCoreDefs.h"
 #include "ActorModifierCoreStack.generated.h"
 
+class UActorModifierCoreComponent;
 class USceneComponent;
 
 /** A modifier stack contains modifiers and is also a modifier by itself */
@@ -34,7 +35,7 @@ public:
 	ACTORMODIFIERCORE_API static FOnModifierUpdated OnModifierMovedDelegate;
 
 	/** Create a new stack by passing the actor and the parent stack if there is one */
-	static UActorModifierCoreStack* Create(AActor* InActor, UActorModifierCoreStack* InParentStack = nullptr);
+	static UActorModifierCoreStack* Create(UActorModifierCoreComponent* InComponent, UActorModifierCoreStack* InParentStack = nullptr);
 
 	/** Gets all modifiers in this stack, does not recurse */
 	TConstArrayView<UActorModifierCoreBase*> GetModifiers() const
@@ -76,10 +77,7 @@ public:
 	ACTORMODIFIERCORE_API TArray<UActorModifierCoreBase*> FindModifiers(const UClass* InSearchClass, const FActorModifierCoreStackSearchOp& InSearchOptions = FActorModifierCoreStackSearchOp::GetDefault()) const;
 
 	/** This is the root actor stack if we do not have any parent stack */
-	bool IsRootStack() const
-	{
-		return !ModifierStack.IsValid();
-	}
+	ACTORMODIFIERCORE_API bool IsRootStack() const;
 
 	/** Execute those function when the stack is restored, before executing it again */
 	ACTORMODIFIERCORE_API void ProcessFunctionOnRestore(const TFunction<void()>& InFunction);
