MotionDesign : ActorModifierCore
- Fix crash issue with modifiers outer rename migration

diff --git a/Engine/Plugins/Experimental/ActorModifierCore/Source/ActorModifierCore/Private/Modifiers/ActorModifierCoreBase.cpp b/Engine/Plugins/Experimental/ActorModifierCore/Source/ActorModifierCore/Private/Modifiers/ActorModifierCoreBase.cpp
--- a/Engine/Plugins/Experimental/ActorModifierCore/Source/ActorModifierCore/Private/Modifiers/ActorModifierCoreBase.cpp
+++ b/Engine/Plugins/Experimental/ActorModifierCore/Source/ActorModifierCore/Private/Modifiers/ActorModifierCoreBase.cpp
@@ -477,7 +477,7 @@ void UActorModifierCoreBase::PostLoad()
 				const EObjectFlags ThisStackFlags = GetFlags();
 
 				LogModifier(FString::Printf(TEXT("Modifier stack migrated to component stack %s %s"), *OwningComponent->GetName(), *ComponentStack->GetName()), true);
-				Rename(TEXT("Trash_ModifierStack"), GetTransientPackage(), RenameFlags);
+				Rename(nullptr, GetTransientPackage(), RenameFlags);
 
 				ComponentStack->Rename(*ThisStackName, nullptr, RenameFlags);
 				ComponentStack->Modifiers = ThisStack->Modifiers;
diff --git a/Engine/Plugins/Experimental/ActorModifierCore/Source/ActorModifierCore/Private/Modifiers/ActorModifierCoreStack.cpp b/Engine/Plugins/Experimental/ActorModifierCore/Source/ActorModifierCore/Private/Modifiers/ActorModifierCoreStack.cpp
--- a/Engine/Plugins/Experimental/ActorModifierCore/Source/ActorModifierCore/Private/Modifiers/ActorModifierCoreStack.cpp
+++ b/Engine/Plugins/Experimental/ActorModifierCore/Source/ActorModifierCore/Private/Modifiers/ActorModifierCoreStack.cpp
@@ -88,13 +88,13 @@ void UActorModifierCoreStack::SetModifierProfiling(bool bInProfiling)
 
 void UActorModifierCoreStack::PostLoad()
 {
-	Super::PostLoad();
-
 	// remove invalid
 	Modifiers.SetNum(Algo::RemoveIf(Modifiers, [](const UActorModifierCoreBase* InModifier)
 	{
-		return !IsValid(InModifier);
+		return !InModifier;
 	}));
+
+	Super::PostLoad();
 }
 
 void UActorModifierCoreStack::OnModifierAdded(EActorModifierCoreEnableReason InReason)
diff --git a/Engine/Plugins/Experimental/ActorModifierCore/Source/ActorModifierCore/Public/Modifiers/ActorModifierCoreComponent.h b/Engine/Plugins/Experimental/ActorModifierCore/Source/ActorModifierCore/Public/Modifiers/ActorModifierCoreComponent.h
--- a/Engine/Plugins/Experimental/ActorModifierCore/Source/ActorModifierCore/Public/Modifiers/ActorModifierCoreComponent.h
+++ b/Engine/Plugins/Experimental/ActorModifierCore/Source/ActorModifierCore/Public/Modifiers/ActorModifierCoreComponent.h
@@ -41,6 +41,6 @@ protected:
 	virtual void TickComponent(float InDeltaTime, ELevelTick InTickType, FActorComponentTickFunction* InThisTickFunction) override;
 
 	/** This is the root stack that will contain all modifiers for this actor */
-	UPROPERTY(VisibleAnywhere, BlueprintReadOnly, Category="Modifiers")
+	UPROPERTY(VisibleAnywhere, BlueprintReadOnly, NoClear, Export, Instanced, Category="Modifiers")
 	TObjectPtr<UActorModifierCoreStack> ModifierStack = nullptr;
 };
\ No newline at end of file
