MotionDesign : ActorModifierCore
- Fix subobject modifier stack initialization skipped after migration causing modifiers not to apply on the actor

diff --git a/Engine/Plugins/Experimental/ActorModifierCore/Source/ActorModifierCore/Private/Modifiers/ActorModifierCoreBase.cpp b/Engine/Plugins/Experimental/ActorModifierCore/Source/ActorModifierCore/Private/Modifiers/ActorModifierCoreBase.cpp
--- a/Engine/Plugins/Experimental/ActorModifierCore/Source/ActorModifierCore/Private/Modifiers/ActorModifierCoreBase.cpp
+++ b/Engine/Plugins/Experimental/ActorModifierCore/Source/ActorModifierCore/Private/Modifiers/ActorModifierCoreBase.cpp
@@ -429,6 +429,11 @@ bool UActorModifierCoreBase::ProcessFunction(TFunctionRef<bool(const UActorModif
 
 void UActorModifierCoreBase::DeferInitializeModifier()
 {
+	if (IsModifierInitialized())
+	{
+		return;
+	}
+
 	// Begin batch operation to avoid updating every time a modifier is loaded
 	UActorModifierCoreStack* Stack = GetRootModifierStack();
 	if (!Stack->IsModifierExecutionLocked() && !Stack->IsModifierStackInitialized())
@@ -476,13 +481,15 @@ void UActorModifierCoreBase::PostLoad()
 				const FString ThisStackName = GetName();
 				const EObjectFlags ThisStackFlags = GetFlags();
 
-				LogModifier(FString::Printf(TEXT("Modifier stack migrated to component stack %s %s"), *OwningComponent->GetName(), *ComponentStack->GetName()), true);
+				LogModifier(FString::Printf(TEXT("Modifier stack migrated to component stack %s %s with %i modifiers"), *OwningComponent->GetName(), *ComponentStack->GetName(), ThisStack->Modifiers.Num()), true);
 				Rename(nullptr, GetTransientPackage(), RenameFlags);
 
 				ComponentStack->Rename(*ThisStackName, nullptr, RenameFlags);
 				ComponentStack->Modifiers = ThisStack->Modifiers;
 				ComponentStack->bModifierProfiling = ThisStack->bModifierProfiling;
 				ComponentStack->SetFlags(ThisStackFlags);
+
+				return;
 			}
 		}
 		// Change outer of modifier to component stack instead of actor
diff --git a/Engine/Plugins/Experimental/ActorModifierCore/Source/ActorModifierCore/Private/Modifiers/ActorModifierCoreComponent.cpp b/Engine/Plugins/Experimental/ActorModifierCore/Source/ActorModifierCore/Private/Modifiers/ActorModifierCoreComponent.cpp
--- a/Engine/Plugins/Experimental/ActorModifierCore/Source/ActorModifierCore/Private/Modifiers/ActorModifierCoreComponent.cpp
+++ b/Engine/Plugins/Experimental/ActorModifierCore/Source/ActorModifierCore/Private/Modifiers/ActorModifierCoreComponent.cpp
@@ -73,6 +73,17 @@ void UActorModifierCoreComponent::OnComponentDestroyed(bool bDestroyingHierarchy
 	}
 }
 
+void UActorModifierCoreComponent::PostLoad()
+{
+	Super::PostLoad();
+
+	// Due to the fact that we replaced stack by sub-object stack, init is needed here as old stack is deleted
+	if (ModifierStack && !ModifierStack->IsModifierInitialized())
+	{
+		ModifierStack->DeferInitializeModifier();
+	}
+}
+
 #if WITH_EDITOR
 void UActorModifierCoreComponent::PostEditUndo()
 {
diff --git a/Engine/Plugins/Experimental/ActorModifierCore/Source/ActorModifierCore/Public/Modifiers/ActorModifierCoreComponent.h b/Engine/Plugins/Experimental/ActorModifierCore/Source/ActorModifierCore/Public/Modifiers/ActorModifierCoreComponent.h
--- a/Engine/Plugins/Experimental/ActorModifierCore/Source/ActorModifierCore/Public/Modifiers/ActorModifierCoreComponent.h
+++ b/Engine/Plugins/Experimental/ActorModifierCore/Source/ActorModifierCore/Public/Modifiers/ActorModifierCoreComponent.h
@@ -33,6 +33,7 @@ protected:
 	//~ End UActorComponent
 
 	//~ Begin UObject
+	virtual void PostLoad() override;
 #if WITH_EDITOR
 	virtual void PostEditUndo() override;
 #endif
