MotionDesign : ActorModifierCore
- Fix static analysis PVS studio warnings from horde

diff --git a/Engine/Plugins/Experimental/ActorModifierCore/Source/ActorModifierCore/Private/Modifiers/ActorModifierCoreBase.cpp b/Engine/Plugins/Experimental/ActorModifierCore/Source/ActorModifierCore/Private/Modifiers/ActorModifierCoreBase.cpp
--- a/Engine/Plugins/Experimental/ActorModifierCore/Source/ActorModifierCore/Private/Modifiers/ActorModifierCoreBase.cpp
+++ b/Engine/Plugins/Experimental/ActorModifierCore/Source/ActorModifierCore/Private/Modifiers/ActorModifierCoreBase.cpp
@@ -462,7 +462,7 @@ void UActorModifierCoreBase::PostLoad()
 
 	const AActor* OwningActor = GetModifiedActor();
 	const UObject* Outer = GetOuter();
-	const UActorModifierCoreComponent* OwningComponent = OwningActor->FindComponentByClass<UActorModifierCoreComponent>();
+	const UActorModifierCoreComponent* OwningComponent = OwningActor ? OwningActor->FindComponentByClass<UActorModifierCoreComponent>() : nullptr;
 
 	if (OwningActor
 		&& OwningComponent
