[Insights] Fixed usage of TRACE_BOOKMARK in several plugins (format string needs to be a const TCHAR array).

diff --git a/Engine/Plugins/Experimental/ActorModifierCore/Source/ActorModifierCore/Private/Modifiers/Profiler/ActorModifierCoreProfiler.cpp b/Engine/Plugins/Experimental/ActorModifierCore/Source/ActorModifierCore/Private/Modifiers/Profiler/ActorModifierCoreProfiler.cpp
--- a/Engine/Plugins/Experimental/ActorModifierCore/Source/ActorModifierCore/Private/Modifiers/Profiler/ActorModifierCoreProfiler.cpp
+++ b/Engine/Plugins/Experimental/ActorModifierCore/Source/ActorModifierCore/Private/Modifiers/Profiler/ActorModifierCoreProfiler.cpp
@@ -33,9 +33,9 @@ void FActorModifierCoreProfiler::SetupProfilingStats()
 void FActorModifierCoreProfiler::BeginProfiling()
 {
 	{
-		TRACE_BOOKMARK(*(TEXT("BeginProfiling_") + GetProfilerTag()))
-		TRACE_BEGIN_REGION(*GetProfilerTag())
+		TRACE_BOOKMARK(TEXT("BeginProfiling_%s"), *GetProfilerTag())
 		TRACE_OBJECT(GetModifier())
+		TRACE_BEGIN_REGION(*GetProfilerTag())
 	}
 
 	ExecutionTimeStart = FPlatformTime::Cycles64();
@@ -75,7 +75,7 @@ void FActorModifierCoreProfiler::EndProfiling()
 			<< EndExecutionEvent.ExecutionCount(ExecutionCount)
 			<< EndExecutionEvent.FrameCountDelta(FrameCountDelta)
 			<< EndExecutionEvent.FrameRateDelta(FrameRateDelta);
-		TRACE_BOOKMARK(*(TEXT("EndProfiling_") + GetProfilerTag()))
+		TRACE_BOOKMARK(TEXT("EndProfiling_%s"), *GetProfilerTag())
 	}
 }
 
diff --git a/Engine/Plugins/Experimental/Enterprise/DatasmithRuntime/Source/Private/DatasmithRuntime.cpp b/Engine/Plugins/Experimental/Enterprise/DatasmithRuntime/Source/Private/DatasmithRuntime.cpp
--- a/Engine/Plugins/Experimental/Enterprise/DatasmithRuntime/Source/Private/DatasmithRuntime.cpp
+++ b/Engine/Plugins/Experimental/Enterprise/DatasmithRuntime/Source/Private/DatasmithRuntime.cpp
@@ -248,7 +248,7 @@ void ADatasmithRuntimeActor::OnCloseDelta()
 
 void ADatasmithRuntimeActor::ApplyNewScene()
 {
-	TRACE_BOOKMARK(TEXT("Load started - %s"), *SceneElement->GetName());
+	TRACE_BOOKMARK(TEXT("Load started - %s"), SceneElement->GetName());
 
 	UE_LOG(LogDatasmithRuntime, Log, TEXT("ADatasmithRuntimeActor::ApplyNewScene"));
 
diff --git a/Engine/Plugins/Experimental/Enterprise/DatasmithRuntime/Source/Private/SceneImporter.cpp b/Engine/Plugins/Experimental/Enterprise/DatasmithRuntime/Source/Private/SceneImporter.cpp
--- a/Engine/Plugins/Experimental/Enterprise/DatasmithRuntime/Source/Private/SceneImporter.cpp
+++ b/Engine/Plugins/Experimental/Enterprise/DatasmithRuntime/Source/Private/SceneImporter.cpp
@@ -581,7 +581,7 @@ namespace DatasmithRuntime
 				}
 			}
 
-			TRACE_BOOKMARK(TEXT("Load complete - %s"), *SceneElement->GetName());
+			TRACE_BOOKMARK(TEXT("Load complete - %s"), SceneElement->GetName());
 
 			OnGoingTasks.Empty();
 
diff --git a/Engine/Plugins/Tests/EditorTests/Source/EditorTests/Private/UnrealEd/EditorTestActorSelection.cpp b/Engine/Plugins/Tests/EditorTests/Source/EditorTests/Private/UnrealEd/EditorTestActorSelection.cpp
--- a/Engine/Plugins/Tests/EditorTests/Source/EditorTests/Private/UnrealEd/EditorTestActorSelection.cpp
+++ b/Engine/Plugins/Tests/EditorTests/Source/EditorTests/Private/UnrealEd/EditorTestActorSelection.cpp
@@ -59,7 +59,7 @@ bool FAssetActorTest::RunTest(const FString& Parameters)
 	// Setup
 	AddCommand(new FCloseAllAssetEditorsCommand());
 	AddCommand(new FFunctionLatentCommand([EditorTestName] {
-		TRACE_BOOKMARK(*(EditorTestName + TEXT(" ProfileBegin")));
+		TRACE_BOOKMARK(TEXT("%s ProfileBegin"), *EditorTestName);
 
 		return true;
 		}));
@@ -97,7 +97,7 @@ bool FAssetActorTest::RunTest(const FString& Parameters)
 		FMemory::Free(TestWorld);
 		FMemory::Free(TestAsset);
 
-		TRACE_BOOKMARK(*(EditorTestName + TEXT(" ProfileEnd")));
+		TRACE_BOOKMARK(TEXT("%s ProfileEnd"), *EditorTestName);
 
 		return true;
 		}));
diff --git a/Engine/Plugins/Tests/EditorTests/Source/EditorTests/Private/UnrealEd/EditorTestHierarchicalISMSelection.cpp b/Engine/Plugins/Tests/EditorTests/Source/EditorTests/Private/UnrealEd/EditorTestHierarchicalISMSelection.cpp
--- a/Engine/Plugins/Tests/EditorTests/Source/EditorTests/Private/UnrealEd/EditorTestHierarchicalISMSelection.cpp
+++ b/Engine/Plugins/Tests/EditorTests/Source/EditorTests/Private/UnrealEd/EditorTestHierarchicalISMSelection.cpp
@@ -62,7 +62,7 @@ bool FAssetHierarchicalISMSelectionTest::RunTest(const FString& Parameters)
 	// Setup
 	AddCommand(new FCloseAllAssetEditorsCommand());
 	AddCommand(new FFunctionLatentCommand([EditorTestName] {
-		TRACE_BOOKMARK(*(EditorTestName + TEXT(" ProfileBegin")));
+		TRACE_BOOKMARK(TEXT("%s ProfileBegin"), *EditorTestName);
 
 		return true;
 		}));
@@ -112,7 +112,7 @@ bool FAssetHierarchicalISMSelectionTest::RunTest(const FString& Parameters)
 		FMemory::Free(TestWorld);
 		FMemory::Free(ISMComponent);
 
-		TRACE_BOOKMARK(*(EditorTestName + TEXT(" ProfileEnd")));
+		TRACE_BOOKMARK(TEXT("%s ProfileEnd"), *EditorTestName);
 
 		return true;
 		}));
