MotionDesign : Modifiers
- Adding CRUD operation support through a blueprint function library
- Adding and exposing getter/setter in motion design modifiers to allow blueprint usage

diff --git a/Engine/Plugins/Experimental/ActorModifierCore/Source/ActorModifierCore/Private/Modifiers/Utilities/ActorModifierCoreLibrary.cpp b/Engine/Plugins/Experimental/ActorModifierCore/Source/ActorModifierCore/Private/Modifiers/Utilities/ActorModifierCoreLibrary.cpp
new file mode 100644
--- /dev/null
+++ b/Engine/Plugins/Experimental/ActorModifierCore/Source/ActorModifierCore/Private/Modifiers/Utilities/ActorModifierCoreLibrary.cpp
@@ -0,0 +1,499 @@
+// Copyright Epic Games, Inc. All Rights Reserved.
+
+#include "Modifiers/Utilities/ActorModifierCoreLibrary.h"
+
+#include "Modifiers/ActorModifierCoreBase.h"
+#include "Modifiers/ActorModifierCoreStack.h"
+#include "Subsystems/ActorModifierCoreSubsystem.h"
+
+DEFINE_LOG_CATEGORY_STATIC(LogActorModifierCoreLibrary, Log, All);
+
+bool UActorModifierCoreLibrary::FindOrAddModifierStack(AActor* InActor, UActorModifierCoreStack*& OutModifierStack, bool bInCreateIfNone)
+{
+	OutModifierStack = nullptr;
+	const UActorModifierCoreSubsystem* ModifierSubsystem = UActorModifierCoreSubsystem::Get();
+
+	if (!IsValid(InActor) || !ModifierSubsystem)
+	{
+		return false;
+	}
+
+	OutModifierStack = ModifierSubsystem->GetActorModifierStack(InActor);
+
+	if (!OutModifierStack && bInCreateIfNone)
+	{
+		OutModifierStack = ModifierSubsystem->AddActorModifierStack(InActor);
+	}
+
+	return !!OutModifierStack;
+}
+
+bool UActorModifierCoreLibrary::InsertModifier(UActorModifierCoreStack* InModifierStack, const FActorModifierCoreInsertOperation& InOperation, UActorModifierCoreBase*& OutNewModifier)
+{
+	OutNewModifier = nullptr;
+	const UActorModifierCoreSubsystem* ModifierSubsystem = UActorModifierCoreSubsystem::Get();
+
+	if (!IsValid(InModifierStack) || !IsValid(InModifierStack->GetModifiedActor()) || !ModifierSubsystem)
+	{
+		return false;
+	}
+
+	FText FailReason = FText::GetEmpty();
+	FActorModifierCoreStackInsertOp InsertOp;
+	InsertOp.InsertPosition = InOperation.InsertPosition;
+	InsertOp.InsertPositionContext = InOperation.InsertPositionContext;
+	InsertOp.FailReason = &FailReason;
+
+	if (!GetModifierNameByClass(InOperation.ModifierClass, InsertOp.NewModifierName))
+	{
+		return false;
+	}
+
+	OutNewModifier = ModifierSubsystem->InsertModifier(InModifierStack, InsertOp);
+
+	if (!FailReason.IsEmpty())
+	{
+		UE_LOG(LogActorModifierCoreLibrary, Warning, TEXT("InsertModifier %s failing reason : %s"), *InsertOp.NewModifierName.ToString(), *FailReason.ToString())
+	}
+
+	return !!OutNewModifier;
+}
+
+bool UActorModifierCoreLibrary::CloneModifier(UActorModifierCoreStack* InModifierStack, const FActorModifierCoreCloneOperation& InOperation, UActorModifierCoreBase*& OutNewModifier)
+{
+	OutNewModifier = nullptr;
+	const UActorModifierCoreSubsystem* ModifierSubsystem = UActorModifierCoreSubsystem::Get();
+
+	if (!IsValid(InModifierStack) || !IsValid(InModifierStack->GetModifiedActor()) || !ModifierSubsystem)
+	{
+		return false;
+	}
+
+	FText FailReason = FText::GetEmpty();
+	FActorModifierCoreStackCloneOp CloneOp;
+	CloneOp.ClonePosition = InOperation.ClonePosition;
+	CloneOp.ClonePositionContext = InOperation.ClonePositionContext;
+	CloneOp.FailReason = &FailReason;
+
+	if (!IsValid(InOperation.CloneModifier))
+	{
+		return false;
+	}
+
+	TArray<UActorModifierCoreBase*> NewModifiers = ModifierSubsystem->CloneModifiers({InOperation.CloneModifier}, InModifierStack, CloneOp);
+
+	if (!FailReason.IsEmpty())
+	{
+		UE_LOG(LogActorModifierCoreLibrary, Warning, TEXT("CloneModifier %s failing reason : %s"), *CloneOp.CloneModifier->GetModifierName().ToString(), *FailReason.ToString())
+	}
+
+	if (NewModifiers.IsEmpty())
+	{
+		return false;
+	}
+
+	OutNewModifier = NewModifiers.Last();
+
+	return !!OutNewModifier;
+}
+
+bool UActorModifierCoreLibrary::MoveModifier(UActorModifierCoreStack* InModifierStack, const FActorModifierCoreMoveOperation& InOperation)
+{
+	const UActorModifierCoreSubsystem* ModifierSubsystem = UActorModifierCoreSubsystem::Get();
+
+	if (!IsValid(InModifierStack) || !IsValid(InModifierStack->GetModifiedActor()) || !ModifierSubsystem)
+	{
+		return false;
+	}
+
+	FText FailReason = FText::GetEmpty();
+	FActorModifierCoreStackMoveOp MoveOp;
+	MoveOp.MoveModifier = InOperation.MoveModifier;
+	MoveOp.MovePosition = InOperation.MovePosition;
+	MoveOp.MovePositionContext = InOperation.MovePositionContext;
+	MoveOp.FailReason = &FailReason;
+
+	if (!IsValid(MoveOp.MoveModifier))
+	{
+		return false;
+	}
+
+	const bool bModifierMoved = ModifierSubsystem->MoveModifier(InModifierStack, MoveOp);
+
+	if (!FailReason.IsEmpty())
+	{
+		UE_LOG(LogActorModifierCoreLibrary, Warning, TEXT("MoveModifier %s failing reason : %s"), *MoveOp.MoveModifier->GetModifierName().ToString(), *FailReason.ToString())
+	}
+
+	return bModifierMoved;
+}
+
+bool UActorModifierCoreLibrary::RemoveModifier(UActorModifierCoreStack* InModifierStack, const FActorModifierCoreRemoveOperation& InOperation)
+{
+	const UActorModifierCoreSubsystem* ModifierSubsystem = UActorModifierCoreSubsystem::Get();
+
+	if (!IsValid(InModifierStack) || !IsValid(InModifierStack->GetModifiedActor()) || !ModifierSubsystem)
+	{
+		return false;
+	}
+
+	FText FailReason = FText::GetEmpty();
+	FActorModifierCoreStackRemoveOp RemoveOp;
+	RemoveOp.RemoveModifier = InOperation.RemoveModifier;
+	RemoveOp.bRemoveDependencies = InOperation.bRemoveDependencies;
+	RemoveOp.FailReason = &FailReason;
+
+	if (!IsValid(RemoveOp.RemoveModifier))
+	{
+		return false;
+	}
+
+	const bool bModifierRemoved = ModifierSubsystem->RemoveModifiers({InOperation.RemoveModifier}, RemoveOp);
+
+	if (!FailReason.IsEmpty())
+	{
+		UE_LOG(LogActorModifierCoreLibrary, Warning, TEXT("RemoveModifier %s failing reason : %s"), *RemoveOp.RemoveModifier->GetModifierName().ToString(), *FailReason.ToString())
+	}
+
+	return bModifierRemoved;
+}
+
+bool UActorModifierCoreLibrary::EnableModifier(UActorModifierCoreBase* InModifier, bool bInState)
+{
+	const UActorModifierCoreSubsystem* ModifierSubsystem = UActorModifierCoreSubsystem::Get();
+
+	if (!IsValid(InModifier) || !IsValid(InModifier->GetModifiedActor()) || !ModifierSubsystem)
+	{
+		return false;
+	}
+
+	return ModifierSubsystem->EnableModifiers({InModifier}, bInState);
+}
+
+bool UActorModifierCoreLibrary::IsModifierEnabled(const UActorModifierCoreBase* InModifier, bool& bOutEnabled)
+{
+	if (!IsValid(InModifier) || !IsValid(InModifier->GetModifiedActor()))
+	{
+		return false;
+	}
+
+	bOutEnabled = InModifier->IsModifierEnabled();
+
+	return true;
+}
+
+bool UActorModifierCoreLibrary::GetModifierStack(const UActorModifierCoreBase* InModifier, UActorModifierCoreStack*& OutModifierStack)
+{
+	OutModifierStack = nullptr;
+
+	if (!IsValid(InModifier))
+	{
+		return false;
+	}
+
+	OutModifierStack = InModifier->GetModifierStack();
+
+	return !!OutModifierStack;
+}
+
+bool UActorModifierCoreLibrary::GetModifierActor(const UActorModifierCoreBase* InModifier, AActor*& OutModifiedActor)
+{
+	OutModifiedActor = nullptr;
+
+	if (!IsValid(InModifier))
+	{
+		return false;
+	}
+
+	OutModifiedActor = InModifier->GetModifiedActor();
+
+	return !!OutModifiedActor;
+}
+
+bool UActorModifierCoreLibrary::GetModifierName(const UActorModifierCoreBase* InModifier, FName& OutModifierName)
+{
+	OutModifierName = NAME_None;
+
+	if (!IsValid(InModifier))
+	{
+		return false;
+	}
+
+	OutModifierName = InModifier->GetModifierName();
+
+	return !OutModifierName.IsNone();
+}
+
+bool UActorModifierCoreLibrary::GetModifierNameByClass(TSubclassOf<UActorModifierCoreBase> InModifierClass, FName& OutModifierName)
+{
+	const UActorModifierCoreSubsystem* ModifierSubsystem = UActorModifierCoreSubsystem::Get();
+	const UClass* ModifierClass = InModifierClass.Get();
+	OutModifierName = NAME_None;
+
+	if (!ModifierClass || !ModifierSubsystem)
+	{
+		return false;
+	}
+
+	OutModifierName = ModifierSubsystem->GetRegisteredModifierName(ModifierClass);
+
+	return !OutModifierName.IsNone();
+}
+
+bool UActorModifierCoreLibrary::GetModifierCategory(const UActorModifierCoreBase* InModifier, FName& OutModifierCategory)
+{
+	OutModifierCategory = NAME_None;
+
+	if (!IsValid(InModifier))
+	{
+		return false;
+	}
+
+	OutModifierCategory = InModifier->GetModifierCategory();
+
+	return !OutModifierCategory.IsNone();
+}
+
+bool UActorModifierCoreLibrary::GetModifierCategoryByClass(TSubclassOf<UActorModifierCoreBase> InModifierClass, FName& OutModifierCategory)
+{
+	const UActorModifierCoreSubsystem* ModifierSubsystem = UActorModifierCoreSubsystem::Get();
+	const UClass* ModifierClass = InModifierClass.Get();
+	OutModifierCategory = NAME_None;
+
+	if (!ModifierClass || !ModifierSubsystem)
+	{
+		return false;
+	}
+
+	const FName ModifierName = ModifierSubsystem->GetRegisteredModifierName(ModifierClass);
+	OutModifierCategory = ModifierSubsystem->GetModifierCategory(ModifierName);
+
+	return !OutModifierCategory.IsNone();
+}
+
+bool UActorModifierCoreLibrary::GetModifierCategories(TSet<FName>& OutModifierCategories)
+{
+	OutModifierCategories.Empty();
+	const UActorModifierCoreSubsystem* ModifierSubsystem = UActorModifierCoreSubsystem::Get();
+
+	if (!ModifierSubsystem)
+	{
+		return false;
+	}
+
+	OutModifierCategories = ModifierSubsystem->GetModifierCategories();
+
+	return true;
+}
+
+bool UActorModifierCoreLibrary::GetModifiersByCategory(FName InCategory, TSet<TSubclassOf<UActorModifierCoreBase>>& OutSupportedModifierClasses)
+{
+	OutSupportedModifierClasses.Empty();
+	const UActorModifierCoreSubsystem* ModifierSubsystem = UActorModifierCoreSubsystem::Get();
+
+	if (InCategory.IsNone() || !ModifierSubsystem)
+	{
+		return false;
+	}
+
+	for (const FName& ModifierName : ModifierSubsystem->GetCategoryModifiers(InCategory))
+	{
+		TSubclassOf<UActorModifierCoreBase> ModifierClass = ModifierSubsystem->GetRegisteredModifierClass(ModifierName);
+
+        if (ModifierClass.Get())
+        {
+			OutSupportedModifierClasses.Add(ModifierClass);
+        }
+	}
+
+	return true;
+}
+
+bool UActorModifierCoreLibrary::GetModifierClass(FName InModifierName, TSubclassOf<UActorModifierCoreBase>& OutModifierClass)
+{
+	OutModifierClass = nullptr;
+	const UActorModifierCoreSubsystem* ModifierSubsystem = UActorModifierCoreSubsystem::Get();
+
+	if (InModifierName.IsNone() || !ModifierSubsystem)
+	{
+		return false;
+	}
+
+	OutModifierClass = ModifierSubsystem->GetRegisteredModifierClass(InModifierName);
+
+	return !!OutModifierClass.Get();
+}
+
+bool UActorModifierCoreLibrary::GetStackModifiers(const UActorModifierCoreStack* InModifierStack, TArray<UActorModifierCoreBase*>& OutModifiers)
+{
+	OutModifiers.Empty();
+
+	if (!IsValid(InModifierStack) || !IsValid(InModifierStack->GetModifiedActor()))
+	{
+		return false;
+	}
+
+	OutModifiers = InModifierStack->GetModifiers();
+
+	return true;
+}
+
+bool UActorModifierCoreLibrary::GetDependentModifiers(UActorModifierCoreBase* InModifier, TSet<UActorModifierCoreBase*>& OutModifiers)
+{
+	OutModifiers.Empty();
+
+	if (!IsValid(InModifier) || !InModifier->GetModifierStack())
+	{
+		return false;
+	}
+
+	const UActorModifierCoreStack* Stack = InModifier->GetModifierStack();
+
+	Stack->GetDependentModifiers(InModifier, OutModifiers);
+
+	return true;
+}
+
+bool UActorModifierCoreLibrary::GetRequiredModifiers(UActorModifierCoreBase* InModifier, TSet<UActorModifierCoreBase*>& OutModifiers)
+{
+	OutModifiers.Empty();
+
+	if (!IsValid(InModifier) || !InModifier->GetModifierStack())
+	{
+		return false;
+	}
+
+	const UActorModifierCoreStack* Stack = InModifier->GetModifierStack();
+
+	Stack->GetRequiredModifiers(InModifier, OutModifiers);
+
+	return true;
+}
+
+bool UActorModifierCoreLibrary::ContainsModifiers(UActorModifierCoreStack* InModifierStack, const FActorModifierCoreSearchOperation& InOperation)
+{
+	TSet<UActorModifierCoreBase*> FoundModifiers;
+	return FindModifiers(InModifierStack, InOperation, FoundModifiers);
+}
+
+bool UActorModifierCoreLibrary::FindModifiers(UActorModifierCoreStack* InModifierStack, const FActorModifierCoreSearchOperation& InOperation, TSet<UActorModifierCoreBase*>& OutFoundModifiers)
+{
+	OutFoundModifiers.Empty();
+
+	if (!IsValid(InModifierStack))
+	{
+		return false;
+	}
+
+	for (UActorModifierCoreBase* Modifier : InOperation.Modifiers)
+	{
+		if (!IsValid(Modifier))
+		{
+			continue;
+		}
+
+		const bool bFoundModifier = InModifierStack->ContainsModifier(Modifier);
+
+		if (bFoundModifier)
+		{
+			OutFoundModifiers.Add(Modifier);
+		}
+
+		if (!bFoundModifier && InOperation.SearchMode == EActorModifierCoreSearchMode::And)
+		{
+			return false;
+		}
+
+		if (bFoundModifier && InOperation.SearchMode == EActorModifierCoreSearchMode::Or)
+		{
+			return true;
+		}
+	}
+
+	for (const TSubclassOf<UActorModifierCoreBase>& ModifierClass : InOperation.ModifierClasses)
+	{
+		const UClass* Class = ModifierClass.Get();
+
+		if (!IsValid(Class))
+		{
+			continue;
+		}
+
+		UActorModifierCoreBase* FoundModifier = InModifierStack->FindModifier(Class);
+
+		if (FoundModifier)
+		{
+			OutFoundModifiers.Add(FoundModifier);
+		}
+
+		if (!FoundModifier && InOperation.SearchMode == EActorModifierCoreSearchMode::And)
+		{
+			return false;
+		}
+
+		if (FoundModifier && InOperation.SearchMode == EActorModifierCoreSearchMode::Or)
+		{
+			return true;
+		}
+	}
+
+	for (const FName& ModifierName : InOperation.ModifierNames)
+	{
+		UActorModifierCoreBase* FoundModifier = InModifierStack->FindModifier(ModifierName);
+
+		if (FoundModifier)
+		{
+			OutFoundModifiers.Add(FoundModifier);
+		}
+
+		if (!FoundModifier && InOperation.SearchMode == EActorModifierCoreSearchMode::And)
+		{
+			return false;
+		}
+
+		if (FoundModifier && InOperation.SearchMode == EActorModifierCoreSearchMode::Or)
+		{
+			return true;
+		}
+	}
+
+	return !OutFoundModifiers.IsEmpty();
+}
+
+bool UActorModifierCoreLibrary::GetSupportedModifiers(AActor* InActor, TSet<TSubclassOf<UActorModifierCoreBase>>& OutSupportedModifierClasses, EActorModifierCoreStackPosition InContextPosition, UActorModifierCoreBase* InContextModifier)
+{
+	const UActorModifierCoreSubsystem* ModifierSubsystem = UActorModifierCoreSubsystem::Get();
+	OutSupportedModifierClasses.Empty();
+
+	if (IsValid(InActor) || !ModifierSubsystem)
+	{
+		return false;
+	}
+
+	for (const FName& ModifierName : ModifierSubsystem->GetAllowedModifiers(InActor, InContextModifier, InContextPosition))
+	{
+		TSubclassOf<UActorModifierCoreBase> ModifierClass = ModifierSubsystem->GetRegisteredModifierClass(ModifierName);
+
+		if (ModifierClass.Get())
+		{
+			OutSupportedModifierClasses.Add(ModifierClass);
+		}
+	}
+
+	return true;
+}
+
+bool UActorModifierCoreLibrary::GetAvailableModifiers(TSet<TSubclassOf<UActorModifierCoreBase>>& OutAvailableModifierClasses)
+{
+	const UActorModifierCoreSubsystem* ModifierSubsystem = UActorModifierCoreSubsystem::Get();
+	OutAvailableModifierClasses.Empty();
+
+	if (!ModifierSubsystem)
+	{
+		return false;
+	}
+
+	OutAvailableModifierClasses = ModifierSubsystem->GetRegisteredModifierClasses();
+
+	return true;
+}
diff --git a/Engine/Plugins/Experimental/ActorModifierCore/Source/ActorModifierCore/Private/Subsystems/ActorModifierCoreSubsystem.cpp b/Engine/Plugins/Experimental/ActorModifierCore/Source/ActorModifierCore/Private/Subsystems/ActorModifierCoreSubsystem.cpp
--- a/Engine/Plugins/Experimental/ActorModifierCore/Source/ActorModifierCore/Private/Subsystems/ActorModifierCoreSubsystem.cpp
+++ b/Engine/Plugins/Experimental/ActorModifierCore/Source/ActorModifierCore/Private/Subsystems/ActorModifierCoreSubsystem.cpp
@@ -170,9 +170,27 @@ FName UActorModifierCoreSubsystem::GetRegisteredModifierName(const UClass* InMod
 	return NAME_None;
 }
 
-TSet<const UClass*> UActorModifierCoreSubsystem::GetRegisteredModifierClasses() const
+TSubclassOf<UActorModifierCoreBase> UActorModifierCoreSubsystem::GetRegisteredModifierClass(FName InModifierName) const
 {
-	TSet<const UClass*> ModifiersClass;
+	TSubclassOf<UActorModifierCoreBase> ModifierClass;
+
+	if (InModifierName.IsNone())
+	{
+		return ModifierClass;
+	}
+
+	if (TSharedRef<FActorModifierCoreMetadata> const* ModifierMetadata = ModifiersMetadata.Find(InModifierName))
+	{
+		ModifierClass = (*ModifierMetadata)->GetClass();
+	}
+
+	return ModifierClass;
+}
+
+TSet<TSubclassOf<UActorModifierCoreBase>> UActorModifierCoreSubsystem::GetRegisteredModifierClasses() const
+{
+	TSet<TSubclassOf<UActorModifierCoreBase>> ModifiersClass;
+	ModifiersClass.Reserve(ModifiersMetadata.Num());
 
 	for (const TPair<FName, TSharedRef<FActorModifierCoreMetadata>>& ModifierMetadataPair : ModifiersMetadata)
 	{
diff --git a/Engine/Plugins/Experimental/ActorModifierCore/Source/ActorModifierCore/Public/Modifiers/ActorModifierCoreComponent.h b/Engine/Plugins/Experimental/ActorModifierCore/Source/ActorModifierCore/Public/Modifiers/ActorModifierCoreComponent.h
--- a/Engine/Plugins/Experimental/ActorModifierCore/Source/ActorModifierCore/Public/Modifiers/ActorModifierCoreComponent.h
+++ b/Engine/Plugins/Experimental/ActorModifierCore/Source/ActorModifierCore/Public/Modifiers/ActorModifierCoreComponent.h
@@ -21,6 +21,7 @@ public:
 
 	UActorModifierCoreComponent();
 
+	UFUNCTION(BlueprintPure, Category="Motion Design|Modifiers")
 	UActorModifierCoreStack* GetModifierStack() const
 	{
 		return ModifierStack;
@@ -42,6 +43,6 @@ protected:
 	virtual void TickComponent(float InDeltaTime, ELevelTick InTickType, FActorComponentTickFunction* InThisTickFunction) override;
 
 	/** This is the root stack that will contain all modifiers for this actor */
-	UPROPERTY(VisibleAnywhere, BlueprintReadOnly, NoClear, Export, Instanced, Category="Modifiers")
+	UPROPERTY(VisibleInstanceOnly, NoClear, Export, Instanced, Category="Modifiers")
 	TObjectPtr<UActorModifierCoreStack> ModifierStack = nullptr;
 };
\ No newline at end of file
diff --git a/Engine/Plugins/Experimental/ActorModifierCore/Source/ActorModifierCore/Public/Modifiers/ActorModifierCoreDefs.h b/Engine/Plugins/Experimental/ActorModifierCore/Source/ActorModifierCore/Public/Modifiers/ActorModifierCoreDefs.h
--- a/Engine/Plugins/Experimental/ActorModifierCore/Source/ActorModifierCore/Public/Modifiers/ActorModifierCoreDefs.h
+++ b/Engine/Plugins/Experimental/ActorModifierCore/Source/ActorModifierCore/Public/Modifiers/ActorModifierCoreDefs.h
@@ -48,6 +48,7 @@ enum class EActorModifierCoreDisableReason : uint8
 };
 
 /** Enumerates valid positions for modifier operations */
+UENUM(BlueprintType)
 enum class EActorModifierCoreStackPosition : uint8
 {
 	Before,
diff --git a/Engine/Plugins/Experimental/ActorModifierCore/Source/ActorModifierCore/Public/Modifiers/ActorModifierCoreStack.h b/Engine/Plugins/Experimental/ActorModifierCore/Source/ActorModifierCore/Public/Modifiers/ActorModifierCoreStack.h
--- a/Engine/Plugins/Experimental/ActorModifierCore/Source/ActorModifierCore/Public/Modifiers/ActorModifierCoreStack.h
+++ b/Engine/Plugins/Experimental/ActorModifierCore/Source/ActorModifierCore/Public/Modifiers/ActorModifierCoreStack.h
@@ -63,6 +63,12 @@ public:
 	/** Gets the last modifier in this stack, does not recurse */
 	ACTORMODIFIERCORE_API UActorModifierCoreBase* GetLastModifier() const;
 
+	/** Gets all modifiers found after this one in the stack that depends on this modifier */
+	bool GetDependentModifiers(UActorModifierCoreBase* InModifier, TSet<UActorModifierCoreBase*>& OutDependentModifiers) const;
+
+	/** Gets all modifiers found before this one in the stack that are required by this modifier */
+	bool GetRequiredModifiers(UActorModifierCoreBase* InModifier, TSet<UActorModifierCoreBase*>& OutDependentModifiers) const;
+
 	/** Check that we have a modifier inside this stack, checks also nested stacks */
 	ACTORMODIFIERCORE_API bool ContainsModifier(const FName& InSearchName, const FActorModifierCoreStackSearchOp& InSearchOptions = FActorModifierCoreStackSearchOp::GetDefault()) const;
 	ACTORMODIFIERCORE_API bool ContainsModifier(const UClass* InSearchClass, const FActorModifierCoreStackSearchOp& InSearchOptions = FActorModifierCoreStackSearchOp::GetDefault()) const;
@@ -116,12 +122,6 @@ protected:
 	/** Checks whether all modifier in this stack are initialized */
 	bool IsModifierStackInitialized() const;
 
-	/** Gets all modifiers found after this one in the stack that depends on this modifier */
-	bool GetDependentModifiers(UActorModifierCoreBase* InModifier, TSet<UActorModifierCoreBase*>& OutDependentModifiers) const;
-
-	/** Gets all modifiers found before this one in the stack that are required by this modifier */
-	bool GetRequiredModifiers(UActorModifierCoreBase* InModifier, TSet<UActorModifierCoreBase*>& OutDependentModifiers) const;
-
 	/** Clone a modifier with options from another stack/actor, returns the newly inserted modifier, supports BATCH operation */
 	UActorModifierCoreBase* CloneModifier(FActorModifierCoreStackCloneOp& InCloneOp);
 
@@ -164,10 +164,6 @@ protected:
 	/** Called when a modifier in the stack is dirty */
 	virtual void OnModifierDirty(UActorModifierCoreBase* DirtyModifier, bool bExecute) override;
 
-	/** Contains actual modifiers in the stack */
-	UPROPERTY(BlueprintReadOnly, VisibleInstanceOnly, NoClear, Export, Instanced, Category = "Modifiers")
-	TArray<TObjectPtr<UActorModifierCoreBase>> Modifiers;
-
 private:
 	/** Sets the stack to receive tick events */
 	virtual void OnModifierCDOSetup(FActorModifierCoreMetadata& InMetadata) override;
@@ -194,6 +190,10 @@ private:
 	/** Checks for any possible modifier optimization within the stack */
 	void CheckModifierOptimization(bool bInInvalidateAll);
 
+	/** Contains actual modifiers in the stack */
+	UPROPERTY(VisibleInstanceOnly, NoClear, Export, Instanced, Category = "Modifiers")
+	TArray<TObjectPtr<UActorModifierCoreBase>> Modifiers;
+
 	/** Contains a copy of modifiers in the stack for this round of execution, useful for restore and for query, can be different from modifiers array */
 	UPROPERTY(Transient, DuplicateTransient, NonTransactional)
 	TArray<TObjectPtr<UActorModifierCoreBase>> CurrentModifiers;
diff --git a/Engine/Plugins/Experimental/ActorModifierCore/Source/ActorModifierCore/Public/Modifiers/Utilities/ActorModifierCoreLibrary.h b/Engine/Plugins/Experimental/ActorModifierCore/Source/ActorModifierCore/Public/Modifiers/Utilities/ActorModifierCoreLibrary.h
new file mode 100644
--- /dev/null
+++ b/Engine/Plugins/Experimental/ActorModifierCore/Source/ActorModifierCore/Public/Modifiers/Utilities/ActorModifierCoreLibrary.h
@@ -0,0 +1,229 @@
+// Copyright Epic Games, Inc. All Rights Reserved.
+
+#pragma once
+
+#include "Kismet/BlueprintFunctionLibrary.h"
+#include "Modifiers/Utilities/ActorModifierCoreLibraryDefs.h"
+#include "ActorModifierCoreLibrary.generated.h"
+
+class UActorModifierCoreBase;
+class UActorModifierCoreStack;
+
+/** Blueprint Create/Read/Update/Delete operations for modifiers */
+UCLASS(MinimalAPI, DisplayName="Motion Design Modifier Library")
+class UActorModifierCoreLibrary : public UBlueprintFunctionLibrary
+{
+	GENERATED_BODY()
+
+public:
+	/**
+	 * Retrieves the modifier stack or create one if none is found
+	 * @param InActor The actor to get the modifier stack from
+	 * @param OutModifierStack [Out] The modifier stack for this actor
+	 * @param bInCreateIfNone Whether to create the modifier stack if none is found
+	 * @return true when a valid modifier stack is returned
+	 */
+	UFUNCTION(BlueprintCallable, Category="Motion Design|Modifiers|Utility")
+	static ACTORMODIFIERCORE_API bool FindOrAddModifierStack(AActor* InActor, UActorModifierCoreStack*& OutModifierStack, bool bInCreateIfNone = false);
+
+	/**
+	 * Creates and insert a new modifier into a modifier stack
+	 * @param InModifierStack The modifier stack to use for the operation
+	 * @param InOperation The data for this operation
+	 * @param OutNewModifier [Out] The newly created modifier
+	 * @return true when a valid modifier is returned
+	 */
+	UFUNCTION(BlueprintCallable, Category="Motion Design|Modifiers|Utility")
+	static ACTORMODIFIERCORE_API bool InsertModifier(UActorModifierCoreStack* InModifierStack, const FActorModifierCoreInsertOperation& InOperation, UActorModifierCoreBase*& OutNewModifier);
+
+	/**
+	 * Clones an existing modifier into a modifier stack
+	 * @param InModifierStack The modifier stack to use for the operation
+	 * @param InOperation The data for this operation
+	 * @param OutNewModifier [Out] The newly created modifier
+	 * @return true when a valid modifier is returned
+	 */
+	UFUNCTION(BlueprintCallable, Category="Motion Design|Modifiers|Utility")
+	static ACTORMODIFIERCORE_API bool CloneModifier(UActorModifierCoreStack* InModifierStack, const FActorModifierCoreCloneOperation& InOperation, UActorModifierCoreBase*& OutNewModifier);
+
+	/**
+	 * Moves an existing modifier into a modifier stack
+	 * @param InModifierStack The modifier stack to use for the operation
+	 * @param InOperation The data for this operation
+	 * @return true when the operation was successful
+	 */
+	UFUNCTION(BlueprintCallable, Category="Motion Design|Modifiers|Utility")
+	static ACTORMODIFIERCORE_API bool MoveModifier(UActorModifierCoreStack* InModifierStack, const FActorModifierCoreMoveOperation& InOperation);
+
+	/**
+	 * Removes an existing modifier from a modifier stack
+	 * @param InModifierStack The modifier stack to use for the operation
+	 * @param InOperation The data for this operation
+	 * @return true when the operation was successful
+	 */
+	UFUNCTION(BlueprintCallable, Category="Motion Design|Modifiers|Utility")
+	static ACTORMODIFIERCORE_API bool RemoveModifier(UActorModifierCoreStack* InModifierStack, const FActorModifierCoreRemoveOperation& InOperation);
+
+	/**
+	 * Sets the state of an existing modifier
+	 * @param InModifier The modifier to use for the operation
+	 * @param bInState The new state for the modifier
+	 * @return true when the operation was successful
+	 */
+	UFUNCTION(BlueprintCallable, Category="Motion Design|Modifiers|Utility")
+	static ACTORMODIFIERCORE_API bool EnableModifier(UActorModifierCoreBase* InModifier, bool bInState);
+
+	/**
+	 * Checks the state of an existing modifier
+	 * @param InModifier The modifier to read from
+	 * @param bOutEnabled The modifier enabled state
+	 * @return true when the operation was successful
+	 */
+	UFUNCTION(BlueprintPure, Category="Motion Design|Modifiers|Utility")
+	static ACTORMODIFIERCORE_API bool IsModifierEnabled(const UActorModifierCoreBase* InModifier, bool& bOutEnabled);
+
+	/**
+	 * Retrieves the modifier stack this modifier is in
+	 * @param InModifier The modifier to read from
+	 * @param OutModifierStack [Out] The modifier stack this modifier belongs to
+	 * @return true when a valid stack is returned
+	 */
+	UFUNCTION(BlueprintPure, Category="Motion Design|Modifiers|Utility")
+	static ACTORMODIFIERCORE_API bool GetModifierStack(const UActorModifierCoreBase* InModifier, UActorModifierCoreStack*& OutModifierStack);
+
+	/**
+	 * Retrieves the actor modified by a modifier
+	 * @param InModifier The modifier to read from
+	 * @param OutModifiedActor [Out] The actor modified by this modifier
+	 * @return true when a valid actor is returned
+	 */
+	UFUNCTION(BlueprintPure, Category="Motion Design|Modifiers|Utility")
+	static ACTORMODIFIERCORE_API bool GetModifierActor(const UActorModifierCoreBase* InModifier, AActor*& OutModifiedActor);
+
+	/**
+	 * Retrieves the modifier name of an existing modifier
+	 * @param InModifier The modifier to read from
+	 * @param OutModifierName [Out] The modifier name
+	 * @return true when a valid name is returned
+	 */
+	UFUNCTION(BlueprintPure, Category="Motion Design|Modifiers|Utility")
+	static ACTORMODIFIERCORE_API bool GetModifierName(const UActorModifierCoreBase* InModifier, FName& OutModifierName);
+
+	/**
+	 * Retrieves the modifier name from a modifier class
+	 * @param InModifierClass The modifier class to resolve the name from
+	 * @param OutModifierName [Out] The modifier name
+	 * @return true when a valid name is returned
+	 */
+	UFUNCTION(BlueprintCallable, Category="Motion Design|Modifiers|Utility")
+	static ACTORMODIFIERCORE_API bool GetModifierNameByClass(TSubclassOf<UActorModifierCoreBase> InModifierClass, FName& OutModifierName);
+
+	/**
+	 * Retrieves the modifier category of an existing modifier
+	 * @param InModifier The modifier to read from
+	 * @param OutModifierCategory [Out] The modifier category
+	 * @return true when a valid category is returned
+	 */
+	UFUNCTION(BlueprintPure, Category="Motion Design|Modifiers|Utility")
+	static ACTORMODIFIERCORE_API bool GetModifierCategory(const UActorModifierCoreBase* InModifier, FName& OutModifierCategory);
+
+	/**
+	 * Retrieves the modifier category from a modifier class
+	 * @param InModifierClass The modifier class to resolve the category from
+	 * @param OutModifierCategory [Out] The modifier category
+	 * @return true when a valid category is returned
+	 */
+	UFUNCTION(BlueprintCallable, Category="Motion Design|Modifiers|Utility")
+	static ACTORMODIFIERCORE_API bool GetModifierCategoryByClass(TSubclassOf<UActorModifierCoreBase> InModifierClass, FName& OutModifierCategory);
+
+	/**
+	 * Retrieves the modifier categories available
+	 * @param OutModifierCategories [Out] The modifier categories registered
+	 * @return true when the operation was successful
+	 */
+	UFUNCTION(BlueprintCallable, Category="Motion Design|Modifiers|Utility")
+	static ACTORMODIFIERCORE_API bool GetModifierCategories(TSet<FName>& OutModifierCategories);
+
+	/**
+	 * Retrieves the modifiers classes by a category
+	 * @param InCategory The modifier category to match
+	 * @param OutSupportedModifierClasses [Out] The modifier classes that match the category
+	 * @return true when the operation was successful
+	 */
+	UFUNCTION(BlueprintCallable, Category="Motion Design|Modifiers|Utility")
+	static ACTORMODIFIERCORE_API bool GetModifiersByCategory(FName InCategory, TSet<TSubclassOf<UActorModifierCoreBase>>& OutSupportedModifierClasses);
+
+	/**
+	 * Retrieves the modifier class from a modifier name
+	 * @param InModifierName The modifier name to resolve
+	 * @param OutModifierClass [Out] The modifier class that matches the name
+	 * @return true when the operation was successful
+	 */
+	UFUNCTION(BlueprintPure, Category="Motion Design|Modifiers|Utility")
+	static ACTORMODIFIERCORE_API bool GetModifierClass(FName InModifierName, TSubclassOf<UActorModifierCoreBase>& OutModifierClass);
+
+	/**
+	 * Retrieves all modifiers from a modifier stack
+	 * @param InModifierStack The modifier stack to read from
+	 * @param OutModifiers [Out] The modifiers contained within the stack
+	 * @return true when the operation was successful
+	 */
+	UFUNCTION(BlueprintCallable, Category="Motion Design|Modifiers|Utility")
+	static ACTORMODIFIERCORE_API bool GetStackModifiers(const UActorModifierCoreStack* InModifierStack, TArray<UActorModifierCoreBase*>& OutModifiers);
+
+	/**
+	 * Retrieves all modifiers found after this one that depends on this modifier
+	 * @param InModifier The modifier that is required by others
+	 * @param OutModifiers [Out] The dependent modifiers that required the modifier
+	 * @return true when the operation was successful
+	 */
+	UFUNCTION(BlueprintCallable, Category="Motion Design|Modifiers|Utility")
+	static ACTORMODIFIERCORE_API bool GetDependentModifiers(UActorModifierCoreBase* InModifier, TSet<UActorModifierCoreBase*>& OutModifiers);
+
+	/**
+	 * Retrieves all modifiers found before this one that are required for this modifier
+	 * @param InModifier The modifier that requires other modifiers
+	 * @param OutModifiers [Out] The required modifiers to use the modifier
+	 * @return true when the operation was successful
+	 */
+	UFUNCTION(BlueprintCallable, Category="Motion Design|Modifiers|Utility")
+	static ACTORMODIFIERCORE_API bool GetRequiredModifiers(UActorModifierCoreBase* InModifier, TSet<UActorModifierCoreBase*>& OutModifiers);
+
+	/**
+	 * Searches modifiers in the stack based on the search options provided
+	 * @param InModifierStack The modifier stack used for the search
+	 * @param InOperation The data for this operation
+	 * @return true when the operation was successful 
+	 */
+	UFUNCTION(BlueprintCallable, Category="Motion Design|Modifiers|Utility")
+	static ACTORMODIFIERCORE_API bool ContainsModifiers(UActorModifierCoreStack* InModifierStack, const FActorModifierCoreSearchOperation& InOperation);
+
+	/**
+	 * Searches modifiers in the stack based on the search options provided
+	 * @param InModifierStack The modifier stack used for the search
+	 * @param InOperation The data for this operation
+	 * @param OutFoundModifiers [Out] The modifiers that matches the search filters
+	 * @return true when the operation was successful
+	 */
+	UFUNCTION(BlueprintCallable, Category="Motion Design|Modifiers|Utility")
+	static ACTORMODIFIERCORE_API bool FindModifiers(UActorModifierCoreStack* InModifierStack, const FActorModifierCoreSearchOperation& InOperation, TSet<UActorModifierCoreBase*>& OutFoundModifiers);
+
+	/**
+	 * Gets all modifier classes supported by this actor at a specific position
+	 * @param InActor The actor to check for compatibility
+	 * @param OutSupportedModifierClasses [Out] The modifier classes available
+	 * @param InContextPosition The context position to insert the modifier
+	 * @param InContextModifier The context modifier for insertion
+	 * @return true when the operation was successful
+	 */
+	UFUNCTION(BlueprintCallable, Category="Motion Design|Modifiers|Utility")
+	static ACTORMODIFIERCORE_API bool GetSupportedModifiers(AActor* InActor, TSet<TSubclassOf<UActorModifierCoreBase>>& OutSupportedModifierClasses, EActorModifierCoreStackPosition InContextPosition = EActorModifierCoreStackPosition::Before, UActorModifierCoreBase* InContextModifier = nullptr);
+
+	/**
+	 * Gets all available modifier classes registered
+	 * @param OutAvailableModifierClasses [Out] The modifier classes registered and available to use
+	 * @return true when the operation was successful
+	 */
+	UFUNCTION(BlueprintCallable, Category="Motion Design|Modifiers|Utility")
+	static ACTORMODIFIERCORE_API bool GetAvailableModifiers(TSet<TSubclassOf<UActorModifierCoreBase>>& OutAvailableModifierClasses);
+};
\ No newline at end of file
diff --git a/Engine/Plugins/Experimental/ActorModifierCore/Source/ActorModifierCore/Public/Modifiers/Utilities/ActorModifierCoreLibraryDefs.h b/Engine/Plugins/Experimental/ActorModifierCore/Source/ActorModifierCore/Public/Modifiers/Utilities/ActorModifierCoreLibraryDefs.h
new file mode 100644
--- /dev/null
+++ b/Engine/Plugins/Experimental/ActorModifierCore/Source/ActorModifierCore/Public/Modifiers/Utilities/ActorModifierCoreLibraryDefs.h
@@ -0,0 +1,92 @@
+// Copyright Epic Games, Inc. All Rights Reserved.
+
+#pragma once
+
+#include "Modifiers/ActorModifierCoreDefs.h"
+#include "ActorModifierCoreLibraryDefs.generated.h"
+
+class UActorModifierCoreBase;
+
+UENUM(BlueprintType)
+enum class EActorModifierCoreSearchMode : uint8
+{
+	/** True when one condition is met */
+	Or,
+	/** True when all condition are met */
+	And
+};
+
+USTRUCT(BlueprintType)
+struct FActorModifierCoreInsertOperation
+{
+	GENERATED_BODY()
+
+	UPROPERTY(BlueprintReadWrite, Category="Motion Design|Modifiers|Utility")
+	TSubclassOf<UActorModifierCoreBase> ModifierClass;
+
+	UPROPERTY(BlueprintReadWrite, Category="Motion Design|Modifiers|Utility")
+	EActorModifierCoreStackPosition InsertPosition = EActorModifierCoreStackPosition::Before;
+
+	UPROPERTY(BlueprintReadWrite, Category="Motion Design|Modifiers|Utility")
+	UActorModifierCoreBase* InsertPositionContext = nullptr;
+};
+
+USTRUCT(BlueprintType)
+struct FActorModifierCoreCloneOperation
+{
+	GENERATED_BODY()
+
+	UPROPERTY(BlueprintReadWrite, Category="Motion Design|Modifiers|Utility")
+	UActorModifierCoreBase* CloneModifier;
+
+	UPROPERTY(BlueprintReadWrite, Category="Motion Design|Modifiers|Utility")
+	EActorModifierCoreStackPosition ClonePosition = EActorModifierCoreStackPosition::Before;
+
+	UPROPERTY(BlueprintReadWrite, Category="Motion Design|Modifiers|Utility")
+	UActorModifierCoreBase* ClonePositionContext = nullptr;
+};
+
+USTRUCT(BlueprintType)
+struct FActorModifierCoreMoveOperation
+{
+	GENERATED_BODY()
+
+	UPROPERTY(BlueprintReadWrite, Category="Motion Design|Modifiers|Utility")
+	UActorModifierCoreBase* MoveModifier;
+
+	UPROPERTY(BlueprintReadWrite, Category="Motion Design|Modifiers|Utility")
+	EActorModifierCoreStackPosition MovePosition = EActorModifierCoreStackPosition::Before;
+
+	UPROPERTY(BlueprintReadWrite, Category="Motion Design|Modifiers|Utility")
+	UActorModifierCoreBase* MovePositionContext = nullptr;
+};
+
+USTRUCT(BlueprintType)
+struct FActorModifierCoreRemoveOperation
+{
+	GENERATED_BODY()
+
+	UPROPERTY(BlueprintReadWrite, Category="Motion Design|Modifiers|Utility")
+	UActorModifierCoreBase* RemoveModifier;
+
+	UPROPERTY(BlueprintReadWrite, Category="Motion Design|Modifiers|Utility")
+	bool bRemoveDependencies = false;
+};
+
+USTRUCT(BlueprintType)
+struct FActorModifierCoreSearchOperation
+{
+	GENERATED_BODY()
+
+	UPROPERTY(BlueprintReadWrite, Category="Motion Design|Modifiers|Utility")
+	EActorModifierCoreSearchMode SearchMode = EActorModifierCoreSearchMode::Or;
+
+	UPROPERTY(BlueprintReadWrite, Category="Motion Design|Modifiers|Utility")
+	TSet<TSubclassOf<UActorModifierCoreBase>> ModifierClasses;
+
+	UPROPERTY(BlueprintReadWrite, Category="Motion Design|Modifiers|Utility")
+	TSet<FName> ModifierNames;
+
+	UPROPERTY(BlueprintReadWrite, Category="Motion Design|Modifiers|Utility")
+	TSet<UActorModifierCoreBase*> Modifiers;
+};
\ No newline at end of file
diff --git a/Engine/Plugins/Experimental/ActorModifierCore/Source/ActorModifierCore/Public/Subsystems/ActorModifierCoreSubsystem.h b/Engine/Plugins/Experimental/ActorModifierCore/Source/ActorModifierCore/Public/Subsystems/ActorModifierCoreSubsystem.h
--- a/Engine/Plugins/Experimental/ActorModifierCore/Source/ActorModifierCore/Public/Subsystems/ActorModifierCoreSubsystem.h
+++ b/Engine/Plugins/Experimental/ActorModifierCore/Source/ActorModifierCore/Public/Subsystems/ActorModifierCoreSubsystem.h
@@ -54,8 +54,11 @@ public:
 	/** Return the modifier name or none from the modifier class provided */
 	ACTORMODIFIERCORE_API FName GetRegisteredModifierName(const UClass* InModifierClass) const;
 
+	/** Return the modifier class from the modifier name provided */
+	ACTORMODIFIERCORE_API TSubclassOf<UActorModifierCoreBase> GetRegisteredModifierClass(FName InModifierName) const;
+
 	/** Return a set of class for all modifiers registered */
-	TSet<const UClass*> GetRegisteredModifierClasses() const;
+	TSet<TSubclassOf<UActorModifierCoreBase>> GetRegisteredModifierClasses() const;
 
 	/** Returns the name of modifiers that are currently registered */
 	ACTORMODIFIERCORE_API TSet<FName> GetRegisteredModifiers() const;
diff --git a/Engine/Plugins/Experimental/Avalanche/Source/AvalancheModifiers/Private/Modifiers/AvaAlignBetweenModifier.cpp b/Engine/Plugins/Experimental/Avalanche/Source/AvalancheModifiers/Private/Modifiers/AvaAlignBetweenModifier.cpp
--- a/Engine/Plugins/Experimental/Avalanche/Source/AvalancheModifiers/Private/Modifiers/AvaAlignBetweenModifier.cpp
+++ b/Engine/Plugins/Experimental/Avalanche/Source/AvalancheModifiers/Private/Modifiers/AvaAlignBetweenModifier.cpp
@@ -123,44 +123,46 @@ TSet<AActor*> UAvaAlignBetweenModifier::GetActors(const bool bEnabledOnly) const
 	return OutActors;
 }
 
-void UAvaAlignBetweenModifier::SetReferenceActors(const TSet<FAvaAlignBetweenWeightedActor>& NewReferenceActors)
+void UAvaAlignBetweenModifier::SetReferenceActors(const TSet<FAvaAlignBetweenWeightedActor>& InReferenceActors)
 {
-	ReferenceActors = NewReferenceActors;
+	ReferenceActors = InReferenceActors;
 	OnReferenceActorsChanged();
 }
 
-void UAvaAlignBetweenModifier::AddReferenceActor(const FAvaAlignBetweenWeightedActor& ReferenceActor)
+bool UAvaAlignBetweenModifier::AddReferenceActor(const FAvaAlignBetweenWeightedActor& InReferenceActor)
 {
 	const AActor* const ModifyActor = GetModifiedActor();
 	if (!IsValid(ModifyActor))
 	{
-		return;
+		return false;
 	}
 
-	if (!ReferenceActor.ActorWeak.IsValid() || ReferenceActor.ActorWeak == ModifyActor)
+	if (!InReferenceActor.ActorWeak.IsValid() || InReferenceActor.ActorWeak == ModifyActor)
 	{
-		return;
+		return false;
 	}
 
 	bool bAlreadyInSet = false;
-	ReferenceActors.Add(ReferenceActor, &bAlreadyInSet);
+	ReferenceActors.Add(InReferenceActor, &bAlreadyInSet);
 
 	if (!bAlreadyInSet)
 	{
 		SetTransformExtensionReferenceActors();
-		
+
 		MarkModifierDirty();
 	}
+
+	return !bAlreadyInSet;
 }
 
-bool UAvaAlignBetweenModifier::RemoveReferenceActor(AActor* const Actor)
+bool UAvaAlignBetweenModifier::RemoveReferenceActor(AActor* const InActor)
 {
-	if (!IsValid(Actor))
+	if (!IsValid(InActor))
 	{
 		return false;
 	}
 
-	if (ReferenceActors.Remove(FAvaAlignBetweenWeightedActor(Actor)) > 0)
+	if (ReferenceActors.Remove(FAvaAlignBetweenWeightedActor(InActor)) > 0)
 	{
 		MarkModifierDirty();
 		
diff --git a/Engine/Plugins/Experimental/Avalanche/Source/AvalancheModifiers/Private/Modifiers/AvaAutoFollowModifier.cpp b/Engine/Plugins/Experimental/Avalanche/Source/AvalancheModifiers/Private/Modifiers/AvaAutoFollowModifier.cpp
--- a/Engine/Plugins/Experimental/Avalanche/Source/AvalancheModifiers/Private/Modifiers/AvaAutoFollowModifier.cpp
+++ b/Engine/Plugins/Experimental/Avalanche/Source/AvalancheModifiers/Private/Modifiers/AvaAutoFollowModifier.cpp
@@ -265,44 +265,44 @@ void UAvaAutoFollowModifier::OnFollowedAxisChanged()
 	MarkModifierDirty();
 }
 
-void UAvaAutoFollowModifier::SetDefaultDistance(const FVector& NewDefaultDistance)
+void UAvaAutoFollowModifier::SetDefaultDistance(const FVector& InDefaultDistance)
 {
-	DefaultDistance = NewDefaultDistance;
+	DefaultDistance = InDefaultDistance;
 
 	MarkModifierDirty();
 }
 
-void UAvaAutoFollowModifier::SetMaxDistance(const FVector& NewMaxDistance)
+void UAvaAutoFollowModifier::SetMaxDistance(const FVector& InMaxDistance)
 {
-	MaxDistance = NewMaxDistance;
+	MaxDistance = InMaxDistance;
 
 	MarkModifierDirty();
 }
 
-void UAvaAutoFollowModifier::SetProgress(const FVector& NewProgress)
+void UAvaAutoFollowModifier::SetProgress(const FVector& InProgress)
 {
-	Progress = NewProgress;
+	Progress = InProgress;
 
 	MarkModifierDirty();
 }
 
-void UAvaAutoFollowModifier::SetFollowedAlignment(const FAvaAnchorAlignment& NewFollowedAlignment)
+void UAvaAutoFollowModifier::SetFollowedAlignment(const FAvaAnchorAlignment& InFollowedAlignment)
 {
-	FollowedAlignment = NewFollowedAlignment;
+	FollowedAlignment = InFollowedAlignment;
 
 	MarkModifierDirty();
 }
 
-void UAvaAutoFollowModifier::SetLocalAlignment(const FAvaAnchorAlignment& NewLocalAlignment)
+void UAvaAutoFollowModifier::SetLocalAlignment(const FAvaAnchorAlignment& InLocalAlignment)
 {
-	LocalAlignment = NewLocalAlignment;
+	LocalAlignment = InLocalAlignment;
 
 	MarkModifierDirty();
 }
 
-void UAvaAutoFollowModifier::SetOffsetAxis(const FVector& NewOffsetAxis)
+void UAvaAutoFollowModifier::SetOffsetAxis(const FVector& InOffsetAxis)
 {
-	OffsetAxis = NewOffsetAxis;
+	OffsetAxis = InOffsetAxis;
 
 	MarkModifierDirty();
 }
diff --git a/Engine/Plugins/Experimental/Avalanche/Source/AvalancheModifiers/Private/Modifiers/AvaDynamicMeshConverterModifier.cpp b/Engine/Plugins/Experimental/Avalanche/Source/AvalancheModifiers/Private/Modifiers/AvaDynamicMeshConverterModifier.cpp
--- a/Engine/Plugins/Experimental/Avalanche/Source/AvalancheModifiers/Private/Modifiers/AvaDynamicMeshConverterModifier.cpp
+++ b/Engine/Plugins/Experimental/Avalanche/Source/AvalancheModifiers/Private/Modifiers/AvaDynamicMeshConverterModifier.cpp
@@ -327,6 +327,40 @@ void UAvaDynamicMeshConverterModifier::SetSourceActorWeak(const TWeakObjectPtr<A
 	OnSourceActorChanged();
 }
 
+void UAvaDynamicMeshConverterModifier::SetComponentTypes(const TSet<EAvaDynamicMeshConverterModifierType>& InTypes)
+{
+	EAvaDynamicMeshConverterModifierType NewComponentType = EAvaDynamicMeshConverterModifierType::None;
+
+	for (const EAvaDynamicMeshConverterModifierType Type : InTypes)
+	{
+		EnumAddFlags(NewComponentType, Type);
+	}
+
+	SetComponentType(static_cast<int32>(NewComponentType));
+}
+
+TSet<EAvaDynamicMeshConverterModifierType> UAvaDynamicMeshConverterModifier::GetComponentTypes() const
+{
+	TSet<EAvaDynamicMeshConverterModifierType> ComponentTypes
+	{
+		EAvaDynamicMeshConverterModifierType::StaticMeshComponent,
+		EAvaDynamicMeshConverterModifierType::DynamicMeshComponent,
+		EAvaDynamicMeshConverterModifierType::SkeletalMeshComponent,
+		EAvaDynamicMeshConverterModifierType::BrushComponent,
+		EAvaDynamicMeshConverterModifierType::ProceduralMeshComponent
+	};
+
+	for (TSet<EAvaDynamicMeshConverterModifierType>::TIterator It(ComponentTypes); It; ++It)
+	{
+		if (!HasFlag(*It))
+		{
+			It.RemoveCurrent();
+		}
+	}
+
+	return ComponentTypes;
+}
+
 void UAvaDynamicMeshConverterModifier::SetComponentType(int32 InComponentType)
 {
 	if (InComponentType == ComponentType)
diff --git a/Engine/Plugins/Experimental/Avalanche/Source/AvalancheModifiers/Private/Modifiers/AvaExtrudeModifier.cpp b/Engine/Plugins/Experimental/Avalanche/Source/AvalancheModifiers/Private/Modifiers/AvaExtrudeModifier.cpp
--- a/Engine/Plugins/Experimental/Avalanche/Source/AvalancheModifiers/Private/Modifiers/AvaExtrudeModifier.cpp
+++ b/Engine/Plugins/Experimental/Avalanche/Source/AvalancheModifiers/Private/Modifiers/AvaExtrudeModifier.cpp
@@ -204,16 +204,13 @@ void UAvaExtrudeModifier::Apply()
 
 void UAvaExtrudeModifier::SetDepth(float InDepth)
 {
-	if (Depth == InDepth)
-	{
-		return;
-	}
+	InDepth = FMath::Max(0, InDepth);
 
-	if (InDepth < 0)
+	if (FMath::IsNearlyEqual(Depth, InDepth))
 	{
 		return;
 	}
-	
+
 	Depth = InDepth;
 	OnDepthChanged();
 }
diff --git a/Engine/Plugins/Experimental/Avalanche/Source/AvalancheModifiers/Private/Modifiers/AvaGridArrangeModifier.cpp b/Engine/Plugins/Experimental/Avalanche/Source/AvalancheModifiers/Private/Modifiers/AvaGridArrangeModifier.cpp
--- a/Engine/Plugins/Experimental/Avalanche/Source/AvalancheModifiers/Private/Modifiers/AvaGridArrangeModifier.cpp
+++ b/Engine/Plugins/Experimental/Avalanche/Source/AvalancheModifiers/Private/Modifiers/AvaGridArrangeModifier.cpp
@@ -45,14 +45,14 @@ void UAvaGridArrangeModifier::SetCount(const FIntPoint& InCount)
 	MarkModifierDirty();
 }
 
-void UAvaGridArrangeModifier::SetSpread(const FVector2D& NewSpread)
+void UAvaGridArrangeModifier::SetSpread(const FVector2D& InSpread)
 {
-	if (Spread == NewSpread)
+	if (Spread == InSpread)
 	{
 		return;
 	}
 
-	Spread = NewSpread;
+	Spread = InSpread;
 	MarkModifierDirty();
 }
 
@@ -78,20 +78,6 @@ void UAvaGridArrangeModifier::SetStartDirection(EAvaGridArrangeDirection InDirec
 	MarkModifierDirty();
 }
 
-void UAvaGridArrangeModifier::GetArrangementLayout(EAvaCorner2D& OutStartCorner, EAvaGridArrangeDirection& OutStartDirection) const
-{
-	OutStartCorner = StartCorner;
-	OutStartDirection = StartDirection;
-}
-
-void UAvaGridArrangeModifier::SetArrangementLayout(const EAvaCorner2D NewStartCorner, const EAvaGridArrangeDirection NewStartDirection)
-{
-	StartCorner = NewStartCorner;
-	StartDirection = NewStartDirection;
-
-	MarkModifierDirty();
-}
-
 void UAvaGridArrangeModifier::OnModifierCDOSetup(FActorModifierCoreMetadata& InMetadata)
 {
 	Super::OnModifierCDOSetup(InMetadata);
diff --git a/Engine/Plugins/Experimental/Avalanche/Source/AvalancheModifiers/Private/Modifiers/AvaLookAtModifier.cpp b/Engine/Plugins/Experimental/Avalanche/Source/AvalancheModifiers/Private/Modifiers/AvaLookAtModifier.cpp
--- a/Engine/Plugins/Experimental/Avalanche/Source/AvalancheModifiers/Private/Modifiers/AvaLookAtModifier.cpp
+++ b/Engine/Plugins/Experimental/Avalanche/Source/AvalancheModifiers/Private/Modifiers/AvaLookAtModifier.cpp
@@ -177,16 +177,16 @@ void UAvaLookAtModifier::OnReferenceActorChanged()
 	}
 }
 
-void UAvaLookAtModifier::SetAxis(const EAvaAxis NewAxis)
+void UAvaLookAtModifier::SetAxis(const EAvaAxis InAxis)
 {
-	Axis = NewAxis;
+	Axis = InAxis;
 
 	MarkModifierDirty();
 }
 
-void UAvaLookAtModifier::SetFlipAxis(const bool bNewFlipAxis)
+void UAvaLookAtModifier::SetFlipAxis(const bool bInFlipAxis)
 {
-	bFlipAxis = bNewFlipAxis;
+	bFlipAxis = bInFlipAxis;
 
 	MarkModifierDirty();
 }
diff --git a/Engine/Plugins/Experimental/Avalanche/Source/AvalancheModifiers/Private/Modifiers/AvaNormalModifier.cpp b/Engine/Plugins/Experimental/Avalanche/Source/AvalancheModifiers/Private/Modifiers/AvaNormalModifier.cpp
--- a/Engine/Plugins/Experimental/Avalanche/Source/AvalancheModifiers/Private/Modifiers/AvaNormalModifier.cpp
+++ b/Engine/Plugins/Experimental/Avalanche/Source/AvalancheModifiers/Private/Modifiers/AvaNormalModifier.cpp
@@ -237,20 +237,20 @@ int32 UAvaNormalModifier::GetPolyGroupLayerIdx() const
 	return Layers.Find(PolyGroupLayer);
 }
 
-void UAvaNormalModifier::SetPolyGroupLayer(FString& InString)
+void UAvaNormalModifier::SetPolyGroupLayer(const FString& InPolyGroupLayer)
 {
-	if (PolyGroupLayer == InString)
+	if (PolyGroupLayer == InPolyGroupLayer)
 	{
 		return;
 	}
 	
 	const TArray<FString> Layers = GetPolyGroupLayers();
-	if (!Layers.Contains(InString))
+	if (!Layers.Contains(InPolyGroupLayer))
 	{
 		return;
 	}
 
-	PolyGroupLayer = InString;
+	PolyGroupLayer = InPolyGroupLayer;
 	OnPolyGroupLayerChanged();
 }
 
diff --git a/Engine/Plugins/Experimental/Avalanche/Source/AvalancheModifiers/Private/Modifiers/AvaRadialArrangeModifier.cpp b/Engine/Plugins/Experimental/Avalanche/Source/AvalancheModifiers/Private/Modifiers/AvaRadialArrangeModifier.cpp
--- a/Engine/Plugins/Experimental/Avalanche/Source/AvalancheModifiers/Private/Modifiers/AvaRadialArrangeModifier.cpp
+++ b/Engine/Plugins/Experimental/Avalanche/Source/AvalancheModifiers/Private/Modifiers/AvaRadialArrangeModifier.cpp
@@ -39,79 +39,79 @@ void UAvaRadialArrangeModifier::PostEditChangeProperty(FPropertyChangedEvent& Pr
 }
 #endif // WITH_EDITOR
 
-void UAvaRadialArrangeModifier::SetCount(const int32 NewCount)
+void UAvaRadialArrangeModifier::SetCount(const int32 InCount)
 {
-	Count = NewCount;
+	Count = InCount;
 
 	MarkModifierDirty();
 }
 
-void UAvaRadialArrangeModifier::SetRings(const int32 NewRings)
+void UAvaRadialArrangeModifier::SetRings(const int32 InRings)
 {
-	Rings = NewRings;
+	Rings = InRings;
 
 	MarkModifierDirty();
 }
 
-void UAvaRadialArrangeModifier::SetInnerRadius(const float NewInnerRadius)
+void UAvaRadialArrangeModifier::SetInnerRadius(const float InInnerRadius)
 {
-	InnerRadius = NewInnerRadius;
+	InnerRadius = InInnerRadius;
 
 	MarkModifierDirty();
 }
 
-void UAvaRadialArrangeModifier::SetOuterRadius(const float NewOuterRadius)
+void UAvaRadialArrangeModifier::SetOuterRadius(const float InOuterRadius)
 {
-	OuterRadius = NewOuterRadius;
+	OuterRadius = InOuterRadius;
 
 	MarkModifierDirty();
 }
 
-void UAvaRadialArrangeModifier::SetStartAngle(const float NewStartAngle)
+void UAvaRadialArrangeModifier::SetStartAngle(const float InStartAngle)
 {
-	StartAngle = NewStartAngle;
+	StartAngle = InStartAngle;
 
 	MarkModifierDirty();
 }
 
-void UAvaRadialArrangeModifier::SetEndAngle(const float NewEndAngle)
+void UAvaRadialArrangeModifier::SetEndAngle(const float InEndAngle)
 {
-	EndAngle = NewEndAngle;
+	EndAngle = InEndAngle;
 
 	MarkModifierDirty();
 }
 
-void UAvaRadialArrangeModifier::SetArrangement(const EAvaRadialArrangement NewArrangement)
+void UAvaRadialArrangeModifier::SetArrangement(const EAvaRadialArrangement InArrangement)
 {
-	Arrangement = NewArrangement;
+	Arrangement = InArrangement;
 
 	MarkModifierDirty();
 }
 
-void UAvaRadialArrangeModifier::SetStartFromOuterRadius(const bool bMewStartFromOuterRadius)
+void UAvaRadialArrangeModifier::SetStartFromOuterRadius(const bool bInStartFromOuterRadius)
 {
-	bStartFromOuterRadius = bMewStartFromOuterRadius;
+	bStartFromOuterRadius = bInStartFromOuterRadius;
 
 	MarkModifierDirty();
 }
 
-void UAvaRadialArrangeModifier::SetOrient(const bool bNewOrient)
+void UAvaRadialArrangeModifier::SetOrient(const bool bInOrient)
 {
-	bOrient = bNewOrient;
+	bOrient = bInOrient;
 
 	MarkModifierDirty();
 }
 
-void UAvaRadialArrangeModifier::SetOrientAxis(const EAvaAxis NewOrientAxis)
+void UAvaRadialArrangeModifier::SetOrientAxis(const EAvaAxis InOrientAxis)
 {
-	OrientAxis = NewOrientAxis;
+	OrientAxis = InOrientAxis;
 
 	MarkModifierDirty();
 }
 
-void UAvaRadialArrangeModifier::SetFlipOrient(const bool bNewFlipOrient)
+void UAvaRadialArrangeModifier::SetFlipOrient(const bool bInFlipOrient)
 {
-	bFlipOrient = bNewFlipOrient;
+	bFlipOrient = bInFlipOrient;
 
 	MarkModifierDirty();
 }
diff --git a/Engine/Plugins/Experimental/Avalanche/Source/AvalancheModifiers/Private/Modifiers/AvaTranslucentPriorityModifier.cpp b/Engine/Plugins/Experimental/Avalanche/Source/AvalancheModifiers/Private/Modifiers/AvaTranslucentPriorityModifier.cpp
--- a/Engine/Plugins/Experimental/Avalanche/Source/AvalancheModifiers/Private/Modifiers/AvaTranslucentPriorityModifier.cpp
+++ b/Engine/Plugins/Experimental/Avalanche/Source/AvalancheModifiers/Private/Modifiers/AvaTranslucentPriorityModifier.cpp
@@ -332,6 +332,11 @@ void UAvaTranslucentPriorityModifier::SetMode(EAvaTranslucentPriorityModifierMod
 	OnModeChanged();
 }
 
+void UAvaTranslucentPriorityModifier::SetCameraActor(ACameraActor* InCameraActor)
+{
+	SetCameraActorWeak(InCameraActor);
+}
+
 void UAvaTranslucentPriorityModifier::SetCameraActorWeak(const TWeakObjectPtr<ACameraActor>& InCameraActor)
 {
 	if (CameraActorWeak == InCameraActor)
diff --git a/Engine/Plugins/Experimental/Avalanche/Source/AvalancheModifiers/Private/Modifiers/AvaVisibilityModifier.cpp b/Engine/Plugins/Experimental/Avalanche/Source/AvalancheModifiers/Private/Modifiers/AvaVisibilityModifier.cpp
--- a/Engine/Plugins/Experimental/Avalanche/Source/AvalancheModifiers/Private/Modifiers/AvaVisibilityModifier.cpp
+++ b/Engine/Plugins/Experimental/Avalanche/Source/AvalancheModifiers/Private/Modifiers/AvaVisibilityModifier.cpp
@@ -161,9 +161,9 @@ void UAvaVisibilityModifier::OnRenderStateUpdated(AActor* InActor, UActorCompone
 	MarkModifierDirty();
 }
 
-void UAvaVisibilityModifier::SetInvertVisibility(const bool bNewInvertVisibility)
+void UAvaVisibilityModifier::SetInvertVisibility(const bool bInInvertVisibility)
 {
-	bInvertVisibility = bNewInvertVisibility;
+	bInvertVisibility = bInInvertVisibility;
 
 	MarkModifierDirty();
 }
diff --git a/Engine/Plugins/Experimental/Avalanche/Source/AvalancheModifiers/Public/Modifiers/AvaAlignBetweenModifier.h b/Engine/Plugins/Experimental/Avalanche/Source/AvalancheModifiers/Public/Modifiers/AvaAlignBetweenModifier.h
--- a/Engine/Plugins/Experimental/Avalanche/Source/AvalancheModifiers/Public/Modifiers/AvaAlignBetweenModifier.h
+++ b/Engine/Plugins/Experimental/Avalanche/Source/AvalancheModifiers/Public/Modifiers/AvaAlignBetweenModifier.h
@@ -44,15 +44,15 @@ struct FAvaAlignBetweenWeightedActor
 	}
 
 	/** An actor that will effect the placement location. */
-	UPROPERTY(EditAnywhere, BlueprintReadWrite, Category = "Motion Design")
+	UPROPERTY(EditInstanceOnly, BlueprintReadWrite, Category="Motion Design")
 	TWeakObjectPtr<AActor> ActorWeak;
 
 	/** How much effect this actor has on the placement location. */
-	UPROPERTY(EditAnywhere, BlueprintReadWrite, Category = "Motion Design", Meta = (ClampMin = "0.0", UIMin = "0.0"))
+	UPROPERTY(EditInstanceOnly, BlueprintReadWrite, Category="Motion Design", Meta = (ClampMin = "0.0", UIMin = "0.0"))
 	float Weight = 0.0f;
 
 	/** If true, will consider this weighted actor when calculating the placement location. */
-	UPROPERTY(EditAnywhere, BlueprintReadWrite, Category = "Motion Design")
+	UPROPERTY(EditInstanceOnly, BlueprintReadWrite, Category="Motion Design")
 	bool bEnabled = false;
 };
 
@@ -66,35 +66,28 @@ class UAvaAlignBetweenModifier : public UAvaBaseModifier
 	GENERATED_BODY()
 
 public:
-	/** Gets all actors from their reference actor structs. */
-	UFUNCTION(BlueprintPure, Category = "Motion Design|Modifiers|AlignBetween")
-	TSet<AActor*> GetActors(const bool bEnabledOnly = false) const;
-
-	/** Returns all valid reference actors that enabled and have a weight greater than 0. */
-	TSet<FAvaAlignBetweenWeightedActor> GetEnabledReferenceActors() const;
-
 	/** Gets all reference actors and their weights. */
-	UFUNCTION(BlueprintPure, Category = "Motion Design|Modifiers|AlignBetween")
+	UFUNCTION(BlueprintPure, Category="Motion Design|Modifiers|AlignBetween")
 	TSet<FAvaAlignBetweenWeightedActor> GetReferenceActors() const
 	{
 		return ReferenceActors;
 	}
 
 	/** Sets all reference actors and their weights. */
-	UFUNCTION(BlueprintCallable, Category = "Motion Design|Modifiers|AlignBetween")
-	AVALANCHEMODIFIERS_API void SetReferenceActors(const TSet<FAvaAlignBetweenWeightedActor>& NewReferenceActors);
+	UFUNCTION(BlueprintCallable, Category="Motion Design|Modifiers|AlignBetween")
+	AVALANCHEMODIFIERS_API void SetReferenceActors(const TSet<FAvaAlignBetweenWeightedActor>& InReferenceActors);
 
 	/** Adds an actor to the reference list. */
-	UFUNCTION(BlueprintCallable, Category = "Motion Design|Modifiers|AlignBetween")
-	void AddReferenceActor(const FAvaAlignBetweenWeightedActor& ReferenceActor);
+	UFUNCTION(BlueprintCallable, Category="Motion Design|Modifiers|AlignBetween")
+	AVALANCHEMODIFIERS_API bool AddReferenceActor(const FAvaAlignBetweenWeightedActor& InReferenceActor);
 
 	/** Removes an actor from the reference list. */
-	UFUNCTION(BlueprintCallable, Category = "Motion Design|Modifiers|AlignBetween")
-	bool RemoveReferenceActor(AActor* const Actor);
+	UFUNCTION(BlueprintCallable, Category="Motion Design|Modifiers|AlignBetween")
+	AVALANCHEMODIFIERS_API bool RemoveReferenceActor(AActor* const InActor);
 
 	/** Finds an actor in the reference list. */
-	UFUNCTION(BlueprintCallable, Category = "Motion Design|Modifiers|AlignBetween")
-	bool FindReferenceActor(AActor* InActor, FAvaAlignBetweenWeightedActor& OutReferenceActor) const;
+	UFUNCTION(BlueprintCallable, Category="Motion Design|Modifiers|AlignBetween")
+	AVALANCHEMODIFIERS_API bool FindReferenceActor(AActor* InActor, FAvaAlignBetweenWeightedActor& OutReferenceActor) const;
 
 protected:
 	//~ Begin UObject
@@ -113,7 +106,6 @@ protected:
 	virtual void OnModifiedActorTransformed() override;
 	//~ End UActorModifierCoreBase
 
-
 	//~ Begin IAvaTransformUpdatedExtension
 	virtual void OnTransformUpdated(AActor* InActor, bool bInParentMoved) override;
 	//~ End IAvaTransformUpdatedExtension
@@ -121,7 +113,13 @@ protected:
 	void OnReferenceActorsChanged();
 	void SetTransformExtensionReferenceActors();
 
+	/** Gets all actors from their reference actor structs. */
+	TSet<AActor*> GetActors(const bool bEnabledOnly = false) const;
+
+	/** Returns all valid reference actors that enabled and have a weight greater than 0. */
+	TSet<FAvaAlignBetweenWeightedActor> GetEnabledReferenceActors() const;
+
 	/** Editable set of reference actors and weights used to calculate the average location for this actor */
-	UPROPERTY(EditAnywhere, BlueprintReadWrite, Setter="SetReferenceActors", Getter="GetReferenceActors", Category="AlignBetween", meta=(AllowPrivateAccess="true"))
+	UPROPERTY(EditInstanceOnly, Setter="SetReferenceActors", Getter="GetReferenceActors", Category="AlignBetween", meta=(AllowPrivateAccess="true"))
 	TSet<FAvaAlignBetweenWeightedActor> ReferenceActors;
 };
diff --git a/Engine/Plugins/Experimental/Avalanche/Source/AvalancheModifiers/Public/Modifiers/AvaAutoFollowModifier.h b/Engine/Plugins/Experimental/Avalanche/Source/AvalancheModifiers/Public/Modifiers/AvaAutoFollowModifier.h
--- a/Engine/Plugins/Experimental/Avalanche/Source/AvalancheModifiers/Public/Modifiers/AvaAutoFollowModifier.h
+++ b/Engine/Plugins/Experimental/Avalanche/Source/AvalancheModifiers/Public/Modifiers/AvaAutoFollowModifier.h
@@ -24,85 +24,85 @@ class UAvaAutoFollowModifier : public UAvaAttachmentBaseModifier
 	GENERATED_BODY()
 
 public:
-	UFUNCTION(BlueprintCallable, Category = "Motion Design|Modifiers|AutoFollow")
+	UFUNCTION(BlueprintCallable, Category="Motion Design|Modifiers|AutoFollow")
 	AVALANCHEMODIFIERS_API void SetReferenceActor(const FAvaSceneTreeActor& InReferenceActor);
 
-	UFUNCTION(BlueprintPure, Category = "Motion Design|Modifiers|AutoFollow")
+	UFUNCTION(BlueprintPure, Category="Motion Design|Modifiers|AutoFollow")
 	const FAvaSceneTreeActor& GetReferenceActor() const
 	{
 		return ReferenceActor;
 	}
 
-	UFUNCTION(BlueprintCallable, Category = "Motion Design|Modifiers|AutoFollow")
+	UFUNCTION(BlueprintCallable, Category="Motion Design|Modifiers|AutoFollow")
 	AVALANCHEMODIFIERS_API void SetFollowedAxis(int32 InFollowedAxis);
 
-	UFUNCTION(BlueprintPure, Category = "Motion Design|Modifiers|AutoFollow")
+	UFUNCTION(BlueprintPure, Category="Motion Design|Modifiers|AutoFollow")
 	int32 GetFollowedAxis() const
 	{
 		return FollowedAxis;
 	}
 
 	/** Sets the distance from this actor to the followed actor. */
-	UFUNCTION(BlueprintCallable, Category = "Motion Design|Modifiers|AutoFollow")
-	AVALANCHEMODIFIERS_API void SetDefaultDistance(const FVector& NewDefaultDistance);
+	UFUNCTION(BlueprintCallable, Category="Motion Design|Modifiers|AutoFollow")
+	AVALANCHEMODIFIERS_API void SetDefaultDistance(const FVector& InDefaultDistance);
 
 	/** Gets the distance from this actor to the followed actor. */
-	UFUNCTION(BlueprintPure, Category = "Motion Design|Modifiers|AutoFollow")
+	UFUNCTION(BlueprintPure, Category="Motion Design|Modifiers|AutoFollow")
 	FVector GetDefaultDistance() const
 	{
 		return DefaultDistance;
 	}
 
 	/** Sets the maximum distance from this actor to the followed actor. */
-	UFUNCTION(BlueprintCallable, Category = "Motion Design|Modifiers|AutoFollow")
-	AVALANCHEMODIFIERS_API void SetMaxDistance(const FVector& NewMaxDistance);
+	UFUNCTION(BlueprintCallable, Category="Motion Design|Modifiers|AutoFollow")
+	AVALANCHEMODIFIERS_API void SetMaxDistance(const FVector& InMaxDistance);
 
 	/** Gets the maximum distance from this actor to the followed actor. */
-	UFUNCTION(BlueprintPure, Category = "Motion Design|Modifiers|AutoFollow")
+	UFUNCTION(BlueprintPure, Category="Motion Design|Modifiers|AutoFollow")
 	FVector GetMaxDistance() const
 	{
 		return MaxDistance;
 	}
 
 	/** Sets the percent % progress from the maximum distance to the default distance. */
-	UFUNCTION(BlueprintCallable, Category = "Motion Design|Modifiers|AutoFollow")
-	AVALANCHEMODIFIERS_API void SetProgress(const FVector& NewProgress);
+	UFUNCTION(BlueprintCallable, Category="Motion Design|Modifiers|AutoFollow")
+	AVALANCHEMODIFIERS_API void SetProgress(const FVector& InProgress);
 
 	/** Gets the percent % progress from the maximum distance to the default distance. */
-	UFUNCTION(BlueprintPure, Category = "Motion Design|Modifiers|AutoFollow")
+	UFUNCTION(BlueprintPure, Category="Motion Design|Modifiers|AutoFollow")
 	FVector GetProgress() const
 	{
 		return Progress;
 	}
 
 	/** Sets the alignment for the followed actor's center. */
-	UFUNCTION(BlueprintCallable, Category = "Motion Design|Modifiers|AutoFollow")
-	AVALANCHEMODIFIERS_API void SetFollowedAlignment(const FAvaAnchorAlignment& NewFollowedAlignment);
+	UFUNCTION(BlueprintCallable, Category="Motion Design|Modifiers|AutoFollow")
+	AVALANCHEMODIFIERS_API void SetFollowedAlignment(const FAvaAnchorAlignment& InFollowedAlignment);
 
 	/** Gets the alignment for the followed actor's center. */
-	UFUNCTION(BlueprintPure, Category = "Motion Design|Modifiers|AutoFollow")
+	UFUNCTION(BlueprintPure, Category="Motion Design|Modifiers|AutoFollow")
 	FAvaAnchorAlignment GetFollowedAlignment() const
 	{
 		return FollowedAlignment;
 	}
 
 	/** Sets the alignment for this actor's center. */
-	UFUNCTION(BlueprintCallable, Category = "Motion Design|Modifiers|AutoFollow")
-	AVALANCHEMODIFIERS_API void SetLocalAlignment(const FAvaAnchorAlignment& NewLocalAlignment);
+	UFUNCTION(BlueprintCallable, Category="Motion Design|Modifiers|AutoFollow")
+	AVALANCHEMODIFIERS_API void SetLocalAlignment(const FAvaAnchorAlignment& InLocalAlignment);
 
 	/** Gets the alignment for this actor's center. */
-	UFUNCTION(BlueprintPure, Category = "Motion Design|Modifiers|AutoFollow")
+	UFUNCTION(BlueprintPure, Category="Motion Design|Modifiers|AutoFollow")
 	FAvaAnchorAlignment GetLocalAlignment() const
 	{
 		return LocalAlignment;
 	}
 
 	/** Sets the axis direction to offset this actor from the followed actor's bounds. */
-	UFUNCTION(BlueprintCallable, Category = "Motion Design|Modifiers|AutoFollow")
-	AVALANCHEMODIFIERS_API void SetOffsetAxis(const FVector& NewOffsetAxis);
+	UFUNCTION(BlueprintCallable, Category="Motion Design|Modifiers|AutoFollow")
+	AVALANCHEMODIFIERS_API void SetOffsetAxis(const FVector& InOffsetAxis);
 
 	/** Gets the axis direction to offset this actor from the followed actor's bounds. */
-	UFUNCTION(BlueprintPure, Category = "Motion Design|Modifiers|AutoFollow")
+	UFUNCTION(BlueprintPure, Category="Motion Design|Modifiers|AutoFollow")
 	FVector GetOffsetAxis() const
 	{
 		return OffsetAxis;
@@ -145,7 +145,7 @@ protected:
 
 	void OnFollowedAxisChanged();
 
-	UPROPERTY(EditAnywhere, BlueprintReadWrite, Setter="SetReferenceActor", Getter="GetReferenceActor", Category="AutoFollow", meta=(ShowOnlyInnerProperties, AllowPrivateAccess="true"))
+	UPROPERTY(EditInstanceOnly, Setter="SetReferenceActor", Getter="GetReferenceActor", Category="AutoFollow", meta=(ShowOnlyInnerProperties, AllowPrivateAccess="true"))
 	FAvaSceneTreeActor ReferenceActor;
 
 	/** The method for finding a reference actor based on it's position in the parent's hierarchy. */
@@ -161,34 +161,34 @@ protected:
 	bool bIgnoreHiddenActors_DEPRECATED = false;
 
 	/** Which axis should we follow */
-	UPROPERTY(EditInstanceOnly, BlueprintReadWrite, Setter="SetFollowedAxis", Getter="GetFollowedAxis", Category="AutoFollow", meta=(Bitmask, BitmaskEnum="/Script/AvalancheModifiers.EAvaModifiersAxis", AllowPrivateAccess="true"))
+	UPROPERTY(EditInstanceOnly, Setter="SetFollowedAxis", Getter="GetFollowedAxis", Category="AutoFollow", meta=(Bitmask, BitmaskEnum="/Script/AvalancheModifiers.EAvaModifiersAxis", AllowPrivateAccess="true"))
 	int32 FollowedAxis = static_cast<int32>(
 		EAvaModifiersAxis::Y |
 		EAvaModifiersAxis::Z
 	);
 
 	/** Based on followed axis, the direction to offset this actor from the followed actor's bounds. */
-	UPROPERTY(EditAnywhere, BlueprintReadWrite, Setter="SetOffsetAxis", Getter="GetOffsetAxis", Interp, Category="AutoFollow", meta=(UIMin="-1.0", UIMax="1.0", AllowPrivateAccess="true"))
+	UPROPERTY(EditInstanceOnly, Setter="SetOffsetAxis", Getter="GetOffsetAxis", Interp, Category="AutoFollow", meta=(UIMin="-1.0", UIMax="1.0", AllowPrivateAccess="true"))
 	FVector OffsetAxis = FVector(0, 1, 0);
 
 	/** The alignment for the followed actor's center. */
-	UPROPERTY(EditAnywhere, BlueprintReadWrite, Setter="SetFollowedAlignment", Getter="GetFollowedAlignment", Category="AutoFollow", meta=(AllowPrivateAccess="true"))
+	UPROPERTY(EditInstanceOnly, Setter="SetFollowedAlignment", Getter="GetFollowedAlignment", Category="AutoFollow", meta=(AllowPrivateAccess="true"))
 	FAvaAnchorAlignment FollowedAlignment;
 
 	/** The alignment for this actor's center. */
-	UPROPERTY(EditAnywhere, BlueprintReadWrite, Setter="SetLocalAlignment", Getter="GetLocalAlignment", Category="AutoFollow", meta=(AllowPrivateAccess="true"))
+	UPROPERTY(EditInstanceOnly, Setter="SetLocalAlignment", Getter="GetLocalAlignment", Category="AutoFollow", meta=(AllowPrivateAccess="true"))
 	FAvaAnchorAlignment LocalAlignment;
 
 	/** The distance from this actor to the followed actor. */
-	UPROPERTY(EditAnywhere, BlueprintReadWrite, Setter="SetDefaultDistance", Getter="GetDefaultDistance", Interp, Category="AutoFollow", meta=(AllowPrivateAccess="true"))
+	UPROPERTY(EditInstanceOnly, Setter="SetDefaultDistance", Getter="GetDefaultDistance", Interp, Category="AutoFollow", meta=(AllowPrivateAccess="true"))
 	FVector DefaultDistance;
 
 	/** The maximum distance from this actor to the followed actor. */
-	UPROPERTY(EditAnywhere, BlueprintReadWrite, Setter="SetMaxDistance", Getter="GetMaxDistance", Interp, Category="AutoFollow", meta=(AllowPrivateAccess="true"))
+	UPROPERTY(EditInstanceOnly, Setter="SetMaxDistance", Getter="GetMaxDistance", Interp, Category="AutoFollow", meta=(AllowPrivateAccess="true"))
 	FVector MaxDistance;
 
 	/** Percent % progress from the maximum distance to the default distance. */
-	UPROPERTY(EditAnywhere, BlueprintReadWrite, Setter="SetProgress", Getter="GetProgress", Interp, Category="AutoFollow", meta=(ClampMin="0.0", UIMin="0.0", ClampMax="100.0", UIMax="100.0", AllowPrivateAccess="true"))
+	UPROPERTY(EditInstanceOnly, Setter="SetProgress", Getter="GetProgress", Interp, Category="AutoFollow", meta=(ClampMin="0.0", UIMin="0.0", ClampMax="100.0", UIMax="100.0", AllowPrivateAccess="true"))
 	FVector Progress;
 
 private:
diff --git a/Engine/Plugins/Experimental/Avalanche/Source/AvalancheModifiers/Public/Modifiers/AvaAutoSizeModifier.h b/Engine/Plugins/Experimental/Avalanche/Source/AvalancheModifiers/Public/Modifiers/AvaAutoSizeModifier.h
--- a/Engine/Plugins/Experimental/Avalanche/Source/AvalancheModifiers/Public/Modifiers/AvaAutoSizeModifier.h
+++ b/Engine/Plugins/Experimental/Avalanche/Source/AvalancheModifiers/Public/Modifiers/AvaAutoSizeModifier.h
@@ -34,48 +34,48 @@ class UAvaAutoSizeModifier : public UAvaGeometryBaseModifier
 	GENERATED_BODY()
 
 public:
-	UFUNCTION(BlueprintCallable, Category = "Motion Design|Modifiers|AutoSize")
+	UFUNCTION(BlueprintCallable, Category="Motion Design|Modifiers|AutoSize")
 	AVALANCHEMODIFIERS_API void SetReferenceActor(const FAvaSceneTreeActor& InReferenceActor);
 
-	UFUNCTION(BlueprintPure, Category = "Motion Design|Modifiers|AutoSize")
+	UFUNCTION(BlueprintPure, Category="Motion Design|Modifiers|AutoSize")
 	const FAvaSceneTreeActor& GetReferenceActor() const
 	{
 		return ReferenceActor;
 	}
 
-	UFUNCTION(BlueprintCallable, Category = "Motion Design|Modifiers|AutoSize")
+	UFUNCTION(BlueprintCallable, Category="Motion Design|Modifiers|AutoSize")
 	AVALANCHEMODIFIERS_API void SetFollowedAxis(int32 InFollowedAxis);
 
-	UFUNCTION(BlueprintPure, Category = "Motion Design|Modifiers|AutoSize")
+	UFUNCTION(BlueprintPure, Category="Motion Design|Modifiers|AutoSize")
 	int32 GetFollowedAxis() const
 	{
 		return FollowedAxis;
 	}
 
 	/** Sets the actor affecting the modifier. This is user selectable if the Reference Container is set to "Other". */
-	UFUNCTION(BlueprintCallable, Category = "Motion Design|Modifiers|AutoSize")
+	UFUNCTION(BlueprintCallable, Category="Motion Design|Modifiers|AutoSize")
 	AVALANCHEMODIFIERS_API void SetPadding(const FMargin& InPadding);
 
 	/** Gets the actor affecting the modifier. This is user selectable if the Reference Container is set to "Other". */
-	UFUNCTION(BlueprintPure, Category = "Motion Design|Modifiers|AutoSize")
+	UFUNCTION(BlueprintPure, Category="Motion Design|Modifiers|AutoSize")
 	const FMargin& GetPadding() const
 	{
 		return Padding;
 	}
 
-	UFUNCTION(BlueprintCallable, Category = "Motion Design|Modifiers|AutoSize")
+	UFUNCTION(BlueprintCallable, Category="Motion Design|Modifiers|AutoSize")
 	AVALANCHEMODIFIERS_API void SetFitMode(const EAvaAutoSizeFitMode InFitMode);
 
-	UFUNCTION(BlueprintPure, Category = "Motion Design|Modifiers|AutoSize")
+	UFUNCTION(BlueprintPure, Category="Motion Design|Modifiers|AutoSize")
 	EAvaAutoSizeFitMode GetFitMode() const
 	{
 		return FitMode;
 	}
 
-	UFUNCTION(BlueprintCallable, Category = "Motion Design|Modifiers|AutoSize")
+	UFUNCTION(BlueprintCallable, Category="Motion Design|Modifiers|AutoSize")
 	AVALANCHEMODIFIERS_API void SetIncludeChildren(bool bInIncludeChildren);
 
-	UFUNCTION(BlueprintPure, Category = "Motion Design|Modifiers|AutoSize")
+	UFUNCTION(BlueprintPure, Category="Motion Design|Modifiers|AutoSize")
 	bool GetIncludeChildren() const
 	{
 		return bIncludeChildren;
@@ -118,7 +118,7 @@ protected:
 
 	void OnReferenceActorChanged();
 
-	UPROPERTY(EditAnywhere, BlueprintReadWrite, Setter="SetReferenceActor", Getter="GetReferenceActor", Category="AutoSize", meta=(ShowOnlyInnerProperties, AllowPrivateAccess="true"))
+	UPROPERTY(EditInstanceOnly, Setter="SetReferenceActor", Getter="GetReferenceActor", Category="AutoSize", meta=(ShowOnlyInnerProperties, AllowPrivateAccess="true"))
 	FAvaSceneTreeActor ReferenceActor;
 
 	/** The method for finding a reference actor based on it's position in the parent's hierarchy. */
@@ -134,21 +134,21 @@ protected:
 	bool bIgnoreHiddenActors_DEPRECATED = false;
 
 	/** Which axis should we follow, if none selected, it will not follow */
-	UPROPERTY(EditInstanceOnly, BlueprintReadWrite, Setter="SetFollowedAxis", Getter="GetFollowedAxis", Category="AutoSize", meta=(Bitmask, BitmaskEnum="/Script/AvalancheModifiers.EAvaModifiersAxis", AllowPrivateAccess="true"))
+	UPROPERTY(EditInstanceOnly, Setter="SetFollowedAxis", Getter="GetFollowedAxis", Category="AutoSize", meta=(Bitmask, BitmaskEnum="/Script/AvalancheModifiers.EAvaModifiersAxis", AllowPrivateAccess="true"))
 	int32 FollowedAxis = static_cast<int32>(
 		EAvaModifiersAxis::Y |
 		EAvaModifiersAxis::Z
 	);
 
 	/* Padding around reference bounds */
-	UPROPERTY(EditAnywhere, BlueprintReadWrite, Setter="SetPadding", Getter="GetPadding", Interp, Category="AutoSize", meta=(AllowPrivateAccess="true"))
+	UPROPERTY(EditInstanceOnly, Setter="SetPadding", Getter="GetPadding", Interp, Category="AutoSize", meta=(AllowPrivateAccess="true"))
 	FMargin Padding = FMargin(0.f);
 
-	UPROPERTY(EditAnywhere, BlueprintReadWrite, Setter="SetFitMode", Getter="GetFitMode", Category="AutoSize", meta=(AllowPrivateAccess="true"))
+	UPROPERTY(EditInstanceOnly, Setter="SetFitMode", Getter="GetFitMode", Category="AutoSize", meta=(AllowPrivateAccess="true"))
 	EAvaAutoSizeFitMode FitMode = EAvaAutoSizeFitMode::WidthAndHeight;
 
 	/** If true, will include children bounds too and compute the new size */
-	UPROPERTY(EditAnywhere, BlueprintReadWrite, Setter="SetIncludeChildren", Getter="GetIncludeChildren", Category="AutoSize", meta=(AllowPrivateAccess="true"))
+	UPROPERTY(EditInstanceOnly, Setter="SetIncludeChildren", Getter="GetIncludeChildren", Category="AutoSize", meta=(AllowPrivateAccess="true"))
 	bool bIncludeChildren = true;
 
 private:
diff --git a/Engine/Plugins/Experimental/Avalanche/Source/AvalancheModifiers/Public/Modifiers/AvaBendModifier.h b/Engine/Plugins/Experimental/Avalanche/Source/AvalancheModifiers/Public/Modifiers/AvaBendModifier.h
--- a/Engine/Plugins/Experimental/Avalanche/Source/AvalancheModifiers/Public/Modifiers/AvaBendModifier.h
+++ b/Engine/Plugins/Experimental/Avalanche/Source/AvalancheModifiers/Public/Modifiers/AvaBendModifier.h
@@ -11,37 +11,55 @@ class UAvaBendModifier : public UAvaGeometryBaseModifier
 	GENERATED_BODY()
 
 public:
+	UFUNCTION(BlueprintCallable, Category="Motion Design|Modifiers|Bend")
 	AVALANCHEMODIFIERS_API void SetAngle(float InAngle);
+
+	UFUNCTION(BlueprintPure, Category="Motion Design|Modifiers|Bend")
 	float GetAngle() const
 	{
 		return Angle;
 	}
 
+	UFUNCTION(BlueprintCallable, Category="Motion Design|Modifiers|Bend")
 	AVALANCHEMODIFIERS_API void SetExtent(float InExtent);
+
+	UFUNCTION(BlueprintPure, Category="Motion Design|Modifiers|Bend")
 	float GetExtent() const
 	{
 		return Extent;
 	}
 
+	UFUNCTION(BlueprintCallable, Category="Motion Design|Modifiers|Bend")
 	AVALANCHEMODIFIERS_API void SetBendPosition(const FVector& InBendPosition);
+
+	UFUNCTION(BlueprintPure, Category="Motion Design|Modifiers|Bend")
 	const FVector& GetBendPosition() const
 	{
 		return BendPosition;
 	}
 
+	UFUNCTION(BlueprintCallable, Category="Motion Design|Modifiers|Bend")
 	AVALANCHEMODIFIERS_API void SetBendRotation(const FRotator& InBendRotation);
+
+	UFUNCTION(BlueprintPure, Category="Motion Design|Modifiers|Bend")
 	const FRotator& GetBendRotation() const
 	{
 		return BendRotation;
 	}
 
+	UFUNCTION(BlueprintCallable, Category="Motion Design|Modifiers|Bend")
 	AVALANCHEMODIFIERS_API void SetSymmetricExtents(bool bInSymmetricExtents);
+
+	UFUNCTION(BlueprintPure, Category="Motion Design|Modifiers|Bend")
 	bool GetSymmetricExtents() const
 	{
 		return bSymmetricExtents;
 	}
 
+	UFUNCTION(BlueprintCallable, Category="Motion Design|Modifiers|Bend")
 	AVALANCHEMODIFIERS_API void SetBidirectional(bool bInBidirectional);
+
+	UFUNCTION(BlueprintPure, Category="Motion Design|Modifiers|Bend")
 	bool GetBidirectional() const
 	{
 		return bBidirectional;
@@ -62,22 +80,22 @@ protected:
 	void OnBendTransformChanged();
 	void OnBendOptionChanged();
 
-	UPROPERTY(EditInstanceOnly, BlueprintReadWrite, Setter="SetBendPosition", Getter="GetBendPosition", Category="Bend", meta=(AllowPrivateAccess="true"))
+	UPROPERTY(EditInstanceOnly, Setter="SetBendPosition", Getter="GetBendPosition", Category="Bend", meta=(AllowPrivateAccess="true"))
 	FVector BendPosition = FVector::ZeroVector;
 
-	UPROPERTY(EditInstanceOnly, BlueprintReadWrite, Setter="SetBendRotation", Getter="GetBendRotation", Category="Bend", meta=(AllowPrivateAccess="true"))
+	UPROPERTY(EditInstanceOnly, Setter="SetBendRotation", Getter="GetBendRotation", Category="Bend", meta=(AllowPrivateAccess="true"))
 	FRotator BendRotation = FRotator::ZeroRotator;
 
-	UPROPERTY(EditInstanceOnly, BlueprintReadWrite, Setter="SetAngle", Getter="GetAngle", Category="Bend", meta=(ClampMin="-360.0", ClampMax="360.0", AllowPrivateAccess="true"))
+	UPROPERTY(EditInstanceOnly, Setter="SetAngle", Getter="GetAngle", Category="Bend", meta=(ClampMin="-360.0", ClampMax="360.0", AllowPrivateAccess="true"))
 	float Angle = 25;
 
-	UPROPERTY(EditInstanceOnly, BlueprintReadWrite, Setter="SetExtent", Getter="GetExtent", Category="Bend", meta=(ClampMin="0", ClampMax="1.0", AllowPrivateAccess="true"))
+	UPROPERTY(EditInstanceOnly, Setter="SetExtent", Getter="GetExtent", Category="Bend", meta=(ClampMin="0", ClampMax="1.0", AllowPrivateAccess="true"))
 	float Extent = 1.0;
 
-	UPROPERTY(EditInstanceOnly, BlueprintReadWrite, Setter="SetSymmetricExtents", Getter="GetSymmetricExtents", Category="Bend", meta=(AllowPrivateAccess="true"))
+	UPROPERTY(EditInstanceOnly, Setter="SetSymmetricExtents", Getter="GetSymmetricExtents", Category="Bend", meta=(AllowPrivateAccess="true"))
 	bool bSymmetricExtents = true;
 
-	UPROPERTY(EditInstanceOnly, BlueprintReadWrite, Setter="SetBidirectional", Getter="GetBidirectional", Category="Bend", meta=(AllowPrivateAccess="true"))
+	UPROPERTY(EditInstanceOnly, Setter="SetBidirectional", Getter="GetBidirectional", Category="Bend", meta=(AllowPrivateAccess="true"))
 	bool bBidirectional = false;
 
 	UPROPERTY()
diff --git a/Engine/Plugins/Experimental/Avalanche/Source/AvalancheModifiers/Public/Modifiers/AvaBevelModifier.h b/Engine/Plugins/Experimental/Avalanche/Source/AvalancheModifiers/Public/Modifiers/AvaBevelModifier.h
--- a/Engine/Plugins/Experimental/Avalanche/Source/AvalancheModifiers/Public/Modifiers/AvaBevelModifier.h
+++ b/Engine/Plugins/Experimental/Avalanche/Source/AvalancheModifiers/Public/Modifiers/AvaBevelModifier.h
@@ -18,19 +18,28 @@ public:
 	static constexpr float MinRoundness = -2;
 	static constexpr float MaxRoundness = 2;
 
+	UFUNCTION(BlueprintCallable, Category="Motion Design|Modifiers|Bevel")
 	AVALANCHEMODIFIERS_API void SetInset(float InBevel);
+
+	UFUNCTION(BlueprintPure, Category="Motion Design|Modifiers|Bevel")
 	float GetInset() const
 	{
 		return Inset;
 	}
 
+	UFUNCTION(BlueprintCallable, Category="Motion Design|Modifiers|Bevel")
 	AVALANCHEMODIFIERS_API void SetIterations(int32 InIterations);
+
+	UFUNCTION(BlueprintPure, Category="Motion Design|Modifiers|Bevel")
 	int32 GetIterations() const
 	{
 		return Iterations;
 	}
 
+	UFUNCTION(BlueprintCallable, Category="Motion Design|Modifiers|Bevel")
 	AVALANCHEMODIFIERS_API void SetRoundness(float InRoundness);
+
+	UFUNCTION(BlueprintPure, Category="Motion Design|Modifiers|Bevel")
 	float GetRoundness() const
 	{
 		return Roundness;
@@ -55,14 +64,14 @@ protected:
 	float GetMaxInsetDistance() const;
 
 	/** Distance used on vertices for beveling, clamped between 0 and (min bound size / 2) */
-	UPROPERTY(EditInstanceOnly, BlueprintReadWrite, Setter, Getter, Category="Bevel", meta=(ClampMin="0", AllowPrivateAccess="true"))
+	UPROPERTY(EditInstanceOnly, Setter, Getter, Category="Bevel", meta=(ClampMin="0", AllowPrivateAccess="true"))
 	float Inset = 1.0f;
 
 	/** Amount of subdivisions applied on the bevel, could affect performance the higher this value gets */
-	UPROPERTY(EditInstanceOnly, BlueprintReadWrite, Setter, Getter, Category="Bevel", meta=(ClampMin="0", ClampMax="10", AllowPrivateAccess="true"))
+	UPROPERTY(EditInstanceOnly, Setter, Getter, Category="Bevel", meta=(ClampMin="0", ClampMax="10", AllowPrivateAccess="true"))
 	int32 Iterations = 0;
 
 	/** Roundness of the beveling when multiple iterations are applied : -2 = inner rounded, 0 = flat, 2 = outer rounded */
-	UPROPERTY(EditInstanceOnly, BlueprintReadWrite, Setter, Getter, Category="Bevel", meta=(ClampMin="-2", ClampMax="2", EditCondition="Iterations > 0", AllowPrivateAccess="true"))
+	UPROPERTY(EditInstanceOnly, Setter, Getter, Category="Bevel", meta=(ClampMin="-2", ClampMax="2", EditCondition="Iterations > 0", AllowPrivateAccess="true"))
 	float Roundness = 0;
 };
diff --git a/Engine/Plugins/Experimental/Avalanche/Source/AvalancheModifiers/Public/Modifiers/AvaBooleanModifier.h b/Engine/Plugins/Experimental/Avalanche/Source/AvalancheModifiers/Public/Modifiers/AvaBooleanModifier.h
--- a/Engine/Plugins/Experimental/Avalanche/Source/AvalancheModifiers/Public/Modifiers/AvaBooleanModifier.h
+++ b/Engine/Plugins/Experimental/Avalanche/Source/AvalancheModifiers/Public/Modifiers/AvaBooleanModifier.h
@@ -32,15 +32,21 @@ class UAvaBooleanModifier : public UAvaGeometryBaseModifier
 
 public:
 	/** This is the min depth threshold needed on the mask for it to work properly */
-	inline static constexpr float MinDepth = UE_KINDA_SMALL_NUMBER * 2;
+	static constexpr float MinDepth = UE_KINDA_SMALL_NUMBER * 2;
 
+	UFUNCTION(BlueprintCallable, Category="Motion Design|Modifiers|Boolean")
 	AVALANCHEMODIFIERS_API void SetMode(EAvaBooleanMode InMode);
+
+	UFUNCTION(BlueprintPure, Category="Motion Design|Modifiers|Boolean")
 	EAvaBooleanMode GetMode() const
 	{
 		return Mode;
 	}
 
+	UFUNCTION(BlueprintCallable, Category="Motion Design|Modifiers|Boolean")
 	AVALANCHEMODIFIERS_API void SetChannel(uint8 InChannel);
+
+	UFUNCTION(BlueprintPure, Category="Motion Design|Modifiers|Boolean")
 	uint8 GetChannel() const
 	{
 		return Channel;
@@ -96,11 +102,11 @@ protected:
 	static void MaskActor(const UAvaBooleanModifier* InTool, const UAvaBooleanModifier* InTarget);
 
 	/** Mode to use when shapes are colliding, none means you will be masked otherwise you are masking */
-	UPROPERTY(EditInstanceOnly, BlueprintReadWrite, Setter="SetMode", Getter="GetMode", Category="Boolean", meta=(AllowPrivateAccess="true"))
+	UPROPERTY(EditInstanceOnly, Setter="SetMode", Getter="GetMode", Category="Boolean", meta=(AllowPrivateAccess="true"))
 	EAvaBooleanMode Mode = EAvaBooleanMode::None;
 
 	/** Channel to only apply this tool on shapes with the same channel */
-	UPROPERTY(EditInstanceOnly, BlueprintReadWrite, Setter="SetChannel", Getter="GetChannel", Category="Boolean", meta=(AllowPrivateAccess="true"))
+	UPROPERTY(EditInstanceOnly, Setter="SetChannel", Getter="GetChannel", Category="Boolean", meta=(AllowPrivateAccess="true"))
 	uint8 Channel = 0;
 
 	UPROPERTY(Transient)
diff --git a/Engine/Plugins/Experimental/Avalanche/Source/AvalancheModifiers/Public/Modifiers/AvaDynamicMeshConverterModifier.h b/Engine/Plugins/Experimental/Avalanche/Source/AvalancheModifiers/Public/Modifiers/AvaDynamicMeshConverterModifier.h
--- a/Engine/Plugins/Experimental/Avalanche/Source/AvalancheModifiers/Public/Modifiers/AvaDynamicMeshConverterModifier.h
+++ b/Engine/Plugins/Experimental/Avalanche/Source/AvalancheModifiers/Public/Modifiers/AvaDynamicMeshConverterModifier.h
@@ -62,6 +62,7 @@ struct FAvaDynamicMeshConverterModifierComponentState
 	/** The default visibility of the converted component in editor */
 	UPROPERTY()
 	bool bComponentVisible = true;
+
 	/** The component converted dynamic mesh*/
 	UE::Geometry::FDynamicMesh3 Mesh;
 };
@@ -72,37 +73,68 @@ class UAvaDynamicMeshConverterModifier : public UAvaGeometryBaseModifier
 	GENERATED_BODY()
 
 public:
+	UFUNCTION(BlueprintCallable, Category="Motion Design|Modifiers|DynamicMeshConverter")
+	void SetSourceActor(AActor* InActor)
+	{
+		SetSourceActorWeak(InActor);
+	}
+
+	UFUNCTION(BlueprintPure, Category="Motion Design|Modifiers|DynamicMeshConverter")
+	AActor* GetSourceActor() const
+	{
+		return SourceActorWeak.Get();
+	}
+
 	AVALANCHEMODIFIERS_API void SetSourceActorWeak(const TWeakObjectPtr<AActor>& InActor);
 	TWeakObjectPtr<AActor> GetSourceActorWeak() const
 	{
 		return SourceActorWeak;
 	}
 
+	UFUNCTION(BlueprintCallable, Category="Motion Design|Modifiers|DynamicMeshConverter")
+	AVALANCHEMODIFIERS_API void SetComponentTypes(const TSet<EAvaDynamicMeshConverterModifierType>& InTypes);
+
+	UFUNCTION(BlueprintPure, Category="Motion Design|Modifiers|DynamicMeshConverter")
+	AVALANCHEMODIFIERS_API TSet<EAvaDynamicMeshConverterModifierType> GetComponentTypes() const;
+
 	AVALANCHEMODIFIERS_API void SetComponentType(int32 InComponentType);
+
 	int32 GetComponentType() const
 	{
 		return ComponentType;
 	}
 
+	UFUNCTION(BlueprintCallable, Category="Motion Design|Modifiers|DynamicMeshConverter")
 	AVALANCHEMODIFIERS_API void SetFilterActorMode(EAvaDynamicMeshConverterModifierFilter InFilter);
+
+	UFUNCTION(BlueprintPure, Category="Motion Design|Modifiers|DynamicMeshConverter")
 	EAvaDynamicMeshConverterModifierFilter GetFilterActorMode() const
 	{
 		return FilterActorMode;
 	}
 
+	UFUNCTION(BlueprintCallable, Category="Motion Design|Modifiers|DynamicMeshConverter")
 	AVALANCHEMODIFIERS_API void SetFilterActorClasses(const TSet<TSubclassOf<AActor>>& InClasses);
+
+	UFUNCTION(BlueprintPure, Category="Motion Design|Modifiers|DynamicMeshConverter")
 	const TSet<TSubclassOf<AActor>>& GetFilterActorClasses() const
 	{
 		return FilterActorClasses;
 	}
 
+	UFUNCTION(BlueprintCallable, Category="Motion Design|Modifiers|DynamicMeshConverter")
 	AVALANCHEMODIFIERS_API void SetIncludeAttachedActors(bool bInInclude);
+
+	UFUNCTION(BlueprintPure, Category="Motion Design|Modifiers|DynamicMeshConverter")
 	bool GetIncludeAttachedActors() const
 	{
 		return bIncludeAttachedActors;
 	}
 
+	UFUNCTION(BlueprintCallable, Category="Motion Design|Modifiers|DynamicMeshConverter")
 	AVALANCHEMODIFIERS_API void SetHideConvertedMesh(bool bInHide);
+
+	UFUNCTION(BlueprintPure, Category="Motion Design|Modifiers|DynamicMeshConverter")
 	bool GetHideConvertedMesh() const
 	{
 		return bHideConvertedMesh;
@@ -148,11 +180,11 @@ protected:
 	void ConvertMesh();
 
 	/** What actor should we copy from, by default is self */
-	UPROPERTY(EditInstanceOnly, BlueprintReadWrite, Setter="SetSourceActorWeak", Getter="GetSourceActorWeak", Category="DynamicMeshConverter", meta=(DisplayName="SourceActor", AllowPrivateAccess="true"))
+	UPROPERTY(EditInstanceOnly, Setter="SetSourceActorWeak", Getter="GetSourceActorWeak", Category="DynamicMeshConverter", meta=(DisplayName="SourceActor", AllowPrivateAccess="true"))
 	TWeakObjectPtr<AActor> SourceActorWeak = GetModifiedActor();
 
 	/** Which components should we take into account for the conversion */
-	UPROPERTY(EditInstanceOnly, BlueprintReadWrite, Setter="SetComponentType", Getter="GetComponentType", Category="DynamicMeshConverter", meta=(Bitmask, BitmaskEnum="/Script/AvalancheModifiers.EAvaDynamicMeshConverterModifierType", AllowPrivateAccess="true"))
+	UPROPERTY(EditInstanceOnly, Setter="SetComponentType", Getter="GetComponentType", Category="DynamicMeshConverter", meta=(Bitmask, BitmaskEnum="/Script/AvalancheModifiers.EAvaDynamicMeshConverterModifierType", AllowPrivateAccess="true"))
 	int32 ComponentType = static_cast<int32>(
 		EAvaDynamicMeshConverterModifierType::StaticMeshComponent |
 		EAvaDynamicMeshConverterModifierType::DynamicMeshComponent |
@@ -161,19 +193,19 @@ protected:
 		EAvaDynamicMeshConverterModifierType::ProceduralMeshComponent);
 
 	/** Actor filter mode : none, include or exclude specific actor class */
-	UPROPERTY(EditInstanceOnly, BlueprintReadWrite, Setter="SetFilterActorMode", Getter="GetFilterActorMode", Category="DynamicMeshConverter", meta=(AllowPrivateAccess="true"))
+	UPROPERTY(EditInstanceOnly, Setter="SetFilterActorMode", Getter="GetFilterActorMode", Category="DynamicMeshConverter", meta=(AllowPrivateAccess="true"))
 	EAvaDynamicMeshConverterModifierFilter FilterActorMode = EAvaDynamicMeshConverterModifierFilter::None;
 
 	/** Actor class to use as filter when gathering actors to convert */
-	UPROPERTY(EditInstanceOnly, BlueprintReadWrite, Setter="SetFilterActorClasses", Getter="GetFilterActorClasses", Category="DynamicMeshConverter", meta=(EditCondition="FilterActorMode != EAvaDynamicMeshConverterModifierFilter::None", AllowPrivateAccess="true"))
+	UPROPERTY(EditInstanceOnly, Setter="SetFilterActorClasses", Getter="GetFilterActorClasses", Category="DynamicMeshConverter", meta=(EditCondition="FilterActorMode != EAvaDynamicMeshConverterModifierFilter::None", AllowPrivateAccess="true"))
 	TSet<TSubclassOf<AActor>> FilterActorClasses;
 
 	/** Checks and convert all attached actors below this actor */
-	UPROPERTY(EditInstanceOnly, BlueprintReadWrite, Setter="SetIncludeAttachedActors", Getter="GetIncludeAttachedActors", Category="DynamicMeshConverter", meta=(AllowPrivateAccess="true"))
+	UPROPERTY(EditInstanceOnly, Setter="SetIncludeAttachedActors", Getter="GetIncludeAttachedActors", Category="DynamicMeshConverter", meta=(AllowPrivateAccess="true"))
 	bool bIncludeAttachedActors = true;
 
 	/** Change visibility of source mesh once they are converted to dynamic mesh, by default will convert itself so hide converted mesh is true */
-	UPROPERTY(EditInstanceOnly, BlueprintReadWrite, Setter="SetHideConvertedMesh", Getter="GetHideConvertedMesh", Category="DynamicMeshConverter", meta=(AllowPrivateAccess="true"))
+	UPROPERTY(EditInstanceOnly, Setter="SetHideConvertedMesh", Getter="GetHideConvertedMesh", Category="DynamicMeshConverter", meta=(AllowPrivateAccess="true"))
 	bool bHideConvertedMesh = true;
 
 	/** Did we create the dynamic mesh component from this modifier or retrieved it */
diff --git a/Engine/Plugins/Experimental/Avalanche/Source/AvalancheModifiers/Public/Modifiers/AvaExtrudeModifier.h b/Engine/Plugins/Experimental/Avalanche/Source/AvalancheModifiers/Public/Modifiers/AvaExtrudeModifier.h
--- a/Engine/Plugins/Experimental/Avalanche/Source/AvalancheModifiers/Public/Modifiers/AvaExtrudeModifier.h
+++ b/Engine/Plugins/Experimental/Avalanche/Source/AvalancheModifiers/Public/Modifiers/AvaExtrudeModifier.h
@@ -5,7 +5,7 @@
 #include "AvaGeometryBaseModifier.h"
 #include "AvaExtrudeModifier.generated.h"
 
-UENUM()
+UENUM(BlueprintType)
 enum class EAvaExtrudeMode : uint8
 {
 	Opposite,
@@ -23,19 +23,28 @@ public:
 	static inline const FName ExtrudePolygroupLayerName = TEXT("ExtrudeSide");
 	static inline const FName BackPolygroupLayerName = TEXT("ExtrudeBack");
 
+	UFUNCTION(BlueprintCallable, Category="Motion Design|Modifiers|Extrude")
 	AVALANCHEMODIFIERS_API void SetDepth(float InDepth);
+
+	UFUNCTION(BlueprintPure, Category="Motion Design|Modifiers|Extrude")
 	float GetDepth() const
 	{
 		return Depth;
 	}
 
+	UFUNCTION(BlueprintCallable, Category="Motion Design|Modifiers|Extrude")
 	AVALANCHEMODIFIERS_API void SetCloseBack(bool bInCloseBack);
+
+	UFUNCTION(BlueprintPure, Category="Motion Design|Modifiers|Extrude")
 	bool GetCloseBack() const
 	{
 		return bCloseBack;
 	}
 
+	UFUNCTION(BlueprintCallable, Category="Motion Design|Modifiers|Extrude")
 	AVALANCHEMODIFIERS_API void SetExtrudeMode(EAvaExtrudeMode InExtrudeMode);
+
+	UFUNCTION(BlueprintPure, Category="Motion Design|Modifiers|Extrude")
 	EAvaExtrudeMode GetExtrudeMode() const
 	{
 		return ExtrudeMode;
@@ -60,11 +69,11 @@ protected:
 	FVector GetExtrudeDirection() const;
 
 	/** Handles mesh depth to extrude primary section */
-	UPROPERTY(EditInstanceOnly, BlueprintReadWrite, Setter="SetDepth", Getter="GetDepth", Category="Extrude", meta=(ClampMin="0", DisplayName="Extrude depth", Units="Centimeters", AllowPrivateAccess="true"))
+	UPROPERTY(EditInstanceOnly, Setter="SetDepth", Getter="GetDepth", Category="Extrude", meta=(ClampMin="0", DisplayName="Extrude depth", Units="Centimeters", AllowPrivateAccess="true"))
 	float Depth = 30.f;
 
 	/** Closes the back of the extrude for a 2D shape for example */
-	UPROPERTY(EditInstanceOnly, BlueprintReadWrite, Setter="SetCloseBack", Getter="GetCloseBack", Category="Extrude", meta=(DisplayName="Close back", AllowPrivateAccess="true"))
+	UPROPERTY(EditInstanceOnly, Setter="SetCloseBack", Getter="GetCloseBack", Category="Extrude", meta=(DisplayName="Close back", AllowPrivateAccess="true"))
 	bool bCloseBack = true;
 
 	/** Moves the mesh in the opposite extrude direction by the depth distance */
@@ -72,6 +81,6 @@ protected:
 	bool bMoveMeshOppositeDirection_DEPRECATED = true;
 
 	/** Specifies the Extrude direction */
-	UPROPERTY(EditInstanceOnly, BlueprintReadWrite, Setter="SetExtrudeMode", Getter="GetExtrudeMode", Category="Extrude", meta=(DisplayName = "Extrude Mode", AllowPrivateAccess="true"))
+	UPROPERTY(EditInstanceOnly, Setter="SetExtrudeMode", Getter="GetExtrudeMode", Category="Extrude", meta=(DisplayName = "Extrude Mode", AllowPrivateAccess="true"))
 	EAvaExtrudeMode ExtrudeMode = EAvaExtrudeMode::Opposite;
 };
diff --git a/Engine/Plugins/Experimental/Avalanche/Source/AvalancheModifiers/Public/Modifiers/AvaGlobalOpacityModifier.h b/Engine/Plugins/Experimental/Avalanche/Source/AvalancheModifiers/Public/Modifiers/AvaGlobalOpacityModifier.h
--- a/Engine/Plugins/Experimental/Avalanche/Source/AvalancheModifiers/Public/Modifiers/AvaGlobalOpacityModifier.h
+++ b/Engine/Plugins/Experimental/Avalanche/Source/AvalancheModifiers/Public/Modifiers/AvaGlobalOpacityModifier.h
@@ -17,7 +17,10 @@ class UAvaGlobalOpacityModifier : public UAvaMaterialParameterModifier
 public:
 	UAvaGlobalOpacityModifier();
 
+	UFUNCTION(BlueprintCallable, Category="Motion Design|Modifiers|GlobalOpacity")
 	AVALANCHEMODIFIERS_API void SetGlobalOpacity(float InOpacity);
+
+	UFUNCTION(BlueprintPure, Category="Motion Design|Modifiers|GlobalOpacity")
 	float GetGlobalOpacity() const
 	{
 		return GlobalOpacity;
@@ -41,6 +44,6 @@ protected:
 	void OnDynamicMaterialValueChanged(UDMMaterialComponent* InComponent, EDMUpdateType InUpdateType);
 
 	/** Global opacity to set on all Material Designer Instances */
-	UPROPERTY(EditAnywhere, BlueprintReadWrite, Setter="SetGlobalOpacity", Getter="GetGlobalOpacity", Category="Material Parameter", meta=(ClampMin="0", ClampMax="1", AllowPrivateAccess="true"))
+	UPROPERTY(EditInstanceOnly, Setter="SetGlobalOpacity", Getter="GetGlobalOpacity", Category="Material Parameter", meta=(ClampMin="0", ClampMax="1", AllowPrivateAccess="true"))
 	float GlobalOpacity = 1.f;
 };
diff --git a/Engine/Plugins/Experimental/Avalanche/Source/AvalancheModifiers/Public/Modifiers/AvaGridArrangeModifier.h b/Engine/Plugins/Experimental/Avalanche/Source/AvalancheModifiers/Public/Modifiers/AvaGridArrangeModifier.h
--- a/Engine/Plugins/Experimental/Avalanche/Source/AvalancheModifiers/Public/Modifiers/AvaGridArrangeModifier.h
+++ b/Engine/Plugins/Experimental/Avalanche/Source/AvalancheModifiers/Public/Modifiers/AvaGridArrangeModifier.h
@@ -25,53 +25,45 @@ class UAvaGridArrangeModifier : public UAvaArrangeBaseModifier
 
 public:
 	/** Sets the width and height of the grid. */
-	UFUNCTION(BlueprintCallable, Category = "Motion Design|Modifiers|GridArrange")
-	AVALANCHEMODIFIERS_API void SetCount(const FIntPoint& NewCount);
+	UFUNCTION(BlueprintCallable, Category="Motion Design|Modifiers|GridArrange")
+	AVALANCHEMODIFIERS_API void SetCount(const FIntPoint& InCount);
 
 	/** Returns the width and height of the grid. */
-	UFUNCTION(BlueprintPure, Category = "Motion Design|Modifiers|GridArrange")
+	UFUNCTION(BlueprintPure, Category="Motion Design|Modifiers|GridArrange")
 	FIntPoint GetCount() const
 	{
 		return Count;
 	}
 
 	/** Sets the distance between each horizontal and vertical child. */
-	UFUNCTION(BlueprintCallable, Category = "Motion Design|Modifiers|GridArrange")
-	AVALANCHEMODIFIERS_API void SetSpread(const FVector2D& NewSpread);
+	UFUNCTION(BlueprintCallable, Category="Motion Design|Modifiers|GridArrange")
+	AVALANCHEMODIFIERS_API void SetSpread(const FVector2D& InSpread);
 
 	/** Returns the distance between each horizontal and vertical child. */
-	UFUNCTION(BlueprintPure, Category = "Motion Design|Modifiers|GridArrange")
+	UFUNCTION(BlueprintPure, Category="Motion Design|Modifiers|GridArrange")
 	FVector2D GetSpread() const
 	{
 		return Spread;
 	}
 
-	UFUNCTION(BlueprintCallable, Category = "Motion Design|Modifiers|GridArrange")
+	UFUNCTION(BlueprintCallable, Category="Motion Design|Modifiers|GridArrange")
 	AVALANCHEMODIFIERS_API void SetStartCorner(EAvaCorner2D InCorner);
 
-	UFUNCTION(BlueprintPure, Category = "Motion Design|Modifiers|GridArrange")
+	UFUNCTION(BlueprintPure, Category="Motion Design|Modifiers|GridArrange")
 	EAvaCorner2D GetStartCorner() const
 	{
 		return StartCorner;
 	}
 
-	UFUNCTION(BlueprintCallable, Category = "Motion Design|Modifiers|GridArrange")
+	UFUNCTION(BlueprintCallable, Category="Motion Design|Modifiers|GridArrange")
 	AVALANCHEMODIFIERS_API void SetStartDirection(EAvaGridArrangeDirection InDirection);
 
-	UFUNCTION(BlueprintPure, Category = "Motion Design|Modifiers|GridArrange")
+	UFUNCTION(BlueprintPure, Category="Motion Design|Modifiers|GridArrange")
 	EAvaGridArrangeDirection GetStartDirection() const
 	{
 		return StartDirection;
 	}
 
-	/** Returns the 2D corner and direction from which to start the arrangement. */
-	UFUNCTION(BlueprintPure, Category = "Motion Design|Modifiers|GridArrange")
-	void GetArrangementLayout(UPARAM(DisplayName = "StartCorner") EAvaCorner2D& OutStartCorner, UPARAM(DisplayName = "Direction") EAvaGridArrangeDirection& OutStartDirection) const;
-
-	/** Sets the 2D corner and direction from which to start the arrangement. */
-	UFUNCTION(BlueprintCallable, Category = "Motion Design|Modifiers|GridArrange")
-	void SetArrangementLayout(const EAvaCorner2D NewStartCorner, const EAvaGridArrangeDirection NewStartDirection);
-
 protected:
 	//~ Begin UObject
 #if WITH_EDITOR
@@ -85,18 +77,18 @@ protected:
 	//~ End UActorModifierCoreBase
 
 	/** The width and height of the grid. */
-	UPROPERTY(EditAnywhere, BlueprintReadWrite, Setter="SetCount", Getter="GetCount", Category="GridArrange", Meta=(ClampMin="1", UIMin="1", AllowPrivateAccess="true"))
+	UPROPERTY(EditInstanceOnly, Setter="SetCount", Getter="GetCount", Category="GridArrange", Meta=(ClampMin="1", UIMin="1", AllowPrivateAccess="true"))
 	FIntPoint Count = FIntPoint(1, 1);
 
 	/** The distance between each horizontal and vertical child. */
-	UPROPERTY(EditAnywhere, BlueprintReadWrite, Setter="SetSpread", Getter="GetSpread", Interp, Category="GridArrange", meta=(AllowPrivateAccess="true"))
+	UPROPERTY(EditInstanceOnly, Setter="SetSpread", Getter="GetSpread", Interp, Category="GridArrange", meta=(AllowPrivateAccess="true"))
 	FVector2D Spread = FVector2D::ZeroVector;
 
 	/** The 2D corner from which to start the arrangement. */
-	UPROPERTY(EditAnywhere, BlueprintReadWrite, Setter="SetStartCorner", Getter="GetStartCorner", Category="GridArrange", meta=(AllowPrivateAccess="true"))
+	UPROPERTY(EditInstanceOnly, Setter="SetStartCorner", Getter="GetStartCorner", Category="GridArrange", meta=(AllowPrivateAccess="true"))
 	EAvaCorner2D StartCorner;
 
 	/** The direction from which to start the arrangement. */
-	UPROPERTY(EditAnywhere, BlueprintReadWrite, Setter="SetStartDirection", Getter="GetStartDirection", Category="GridArrange", meta=(AllowPrivateAccess="true"))
+	UPROPERTY(EditInstanceOnly, Setter="SetStartDirection", Getter="GetStartDirection", Category="GridArrange", meta=(AllowPrivateAccess="true"))
 	EAvaGridArrangeDirection StartDirection;
 };
diff --git a/Engine/Plugins/Experimental/Avalanche/Source/AvalancheModifiers/Public/Modifiers/AvaHideEmptyModifier.h b/Engine/Plugins/Experimental/Avalanche/Source/AvalancheModifiers/Public/Modifiers/AvaHideEmptyModifier.h
--- a/Engine/Plugins/Experimental/Avalanche/Source/AvalancheModifiers/Public/Modifiers/AvaHideEmptyModifier.h
+++ b/Engine/Plugins/Experimental/Avalanche/Source/AvalancheModifiers/Public/Modifiers/AvaHideEmptyModifier.h
@@ -13,13 +13,28 @@ class UAvaHideEmptyModifier : public UAvaArrangeBaseModifier
 	GENERATED_BODY()
 
 public:
+	UFUNCTION(BlueprintCallable, Category="Motion Design|Modifiers|HideEmpty")
+	AVALANCHEMODIFIERS_API void SetContainerActor(AActor* InActor)
+	{
+		SetContainerActorWeak(InActor);
+	}
+
+	UFUNCTION(BlueprintPure, Category="Motion Design|Modifiers|HideEmpty")
+	AActor* GetContainerActor() const
+	{
+		return ContainerActorWeak.Get();
+	}
+
 	AVALANCHEMODIFIERS_API void SetContainerActorWeak(TWeakObjectPtr<AActor> InContainer);
 	TWeakObjectPtr<AActor> GetContainerActorWeak() const
 	{
 		return ContainerActorWeak;
 	}
 
+	UFUNCTION(BlueprintCallable, Category="Motion Design|Modifiers|HideEmpty")
 	AVALANCHEMODIFIERS_API void SetInvertVisibility(bool bInInvert);
+
+	UFUNCTION(BlueprintPure, Category="Motion Design|Modifiers|HideEmpty")
 	bool GetInvertVisibility() const
 	{
 		return bInvertVisibility;
@@ -51,11 +66,11 @@ protected:
 	//~ End IAvaSceneTreeUpdateModifierExtension
 
 	/** The container to hide when text is empty, by default self */
-	UPROPERTY(EditAnywhere, BlueprintReadWrite, Setter, Getter, Category="HideEmpty", meta=(DisplayName="ContainerActor", AllowPrivateAccess="true"))
+	UPROPERTY(EditInstanceOnly, Setter, Getter, Category="HideEmpty", meta=(DisplayName="ContainerActor", AllowPrivateAccess="true"))
 	TWeakObjectPtr<AActor> ContainerActorWeak;
 
 	/** Invert the behaviour and visibility of the container if text is empty */
-	UPROPERTY(EditAnywhere, BlueprintReadWrite, Setter="SetInvertVisibility", Getter="GetInvertVisibility", Category="HideEmpty", meta=(AllowPrivateAccess="true"))
+	UPROPERTY(EditInstanceOnly, Setter="SetInvertVisibility", Getter="GetInvertVisibility", Category="HideEmpty", meta=(AllowPrivateAccess="true"))
 	bool bInvertVisibility = false;
 
 private:
diff --git a/Engine/Plugins/Experimental/Avalanche/Source/AvalancheModifiers/Public/Modifiers/AvaJustifyModifier.h b/Engine/Plugins/Experimental/Avalanche/Source/AvalancheModifiers/Public/Modifiers/AvaJustifyModifier.h
--- a/Engine/Plugins/Experimental/Avalanche/Source/AvalancheModifiers/Public/Modifiers/AvaJustifyModifier.h
+++ b/Engine/Plugins/Experimental/Avalanche/Source/AvalancheModifiers/Public/Modifiers/AvaJustifyModifier.h
@@ -44,55 +44,55 @@ class UAvaJustifyModifier : public UAvaArrangeBaseModifier
 	GENERATED_BODY()
 
 public:
-	UFUNCTION(BlueprintCallable, Category = "Justify")
+	UFUNCTION(BlueprintCallable, Category="Motion Design|Modifiers|Justify")
 	AVALANCHEMODIFIERS_API void SetHorizontalAlignment(EAvaJustifyHorizontal InHorizontalAlignment);
 
-	UFUNCTION(BlueprintPure, Category = "Justify")
+	UFUNCTION(BlueprintPure, Category="Motion Design|Modifiers|Justify")
 	EAvaJustifyHorizontal GetHorizontalAlignment() const
 	{
 		return HorizontalAlignment;
 	}
 
-	UFUNCTION(BlueprintCallable, Category = "Justify")
+	UFUNCTION(BlueprintCallable, Category="Motion Design|Modifiers|Justify")
 	AVALANCHEMODIFIERS_API void SetVerticalAlignment(EAvaJustifyVertical InVerticalAlignment);
 
-	UFUNCTION(BlueprintPure, Category = "Justify")
+	UFUNCTION(BlueprintPure, Category="Motion Design|Modifiers|Justify")
 	EAvaJustifyVertical GetVerticalAlignment() const
 	{
 		return VerticalAlignment;
 	}
 
-	UFUNCTION(BlueprintCallable, Category = "Justify")
+	UFUNCTION(BlueprintCallable, Category="Motion Design|Modifiers|Justify")
 	AVALANCHEMODIFIERS_API void SetDepthAlignment(EAvaJustifyDepth InDepthAlignment);
 
-	UFUNCTION(BlueprintPure, Category = "Justify")
+	UFUNCTION(BlueprintPure, Category="Motion Design|Modifiers|Justify")
 	EAvaJustifyDepth GetDepthAlignment() const
 	{
 		return DepthAlignment;
 	}
 
-	UFUNCTION(BlueprintCallable, Category = "Justify")
+	UFUNCTION(BlueprintCallable, Category="Motion Design|Modifiers|Justify")
 	AVALANCHEMODIFIERS_API void SetHorizontalAnchor(float InHorizontalAnchor);
 
-	UFUNCTION(BlueprintPure, Category = "Justify")
+	UFUNCTION(BlueprintPure, Category="Motion Design|Modifiers|Justify")
 	float GetHorizontalAnchor() const
 	{
 		return HorizontalAnchor;
 	}
 
-	UFUNCTION(BlueprintCallable, Category = "Justify")
+	UFUNCTION(BlueprintCallable, Category="Motion Design|Modifiers|Justify")
 	AVALANCHEMODIFIERS_API void SetVerticalAnchor(float InVerticalAnchor);
 
-	UFUNCTION(BlueprintPure, Category = "Justify")
+	UFUNCTION(BlueprintPure, Category="Motion Design|Modifiers|Justify")
 	float GetVerticalAnchor() const
 	{
 		return VerticalAnchor;
 	}
 
-	UFUNCTION(BlueprintCallable, Category = "Justify")
+	UFUNCTION(BlueprintCallable, Category="Motion Design|Modifiers|Justify")
 	AVALANCHEMODIFIERS_API void SetDepthAnchor(float InDepthAnchor);
 
-	UFUNCTION(BlueprintPure, Category = "Justify")
+	UFUNCTION(BlueprintPure, Category="Motion Design|Modifiers|Justify")
 	float GetDepthAnchor() const
 	{
 		return DepthAnchor;
@@ -138,22 +138,22 @@ protected:
 	FVector GetAnchorOffset() const;
 	FVector GetAlignmentOffset(const FVector& InExtent) const;
 
-	UPROPERTY(EditAnywhere, BlueprintReadWrite, Setter="SetHorizontalAlignment", Getter="GetHorizontalAlignment", Category="Justify", meta=(AllowPrivateAccess="true"))
+	UPROPERTY(EditInstanceOnly, Setter="SetHorizontalAlignment", Getter="GetHorizontalAlignment", Category="Justify", meta=(AllowPrivateAccess="true"))
 	EAvaJustifyHorizontal HorizontalAlignment;
 
-	UPROPERTY(EditAnywhere, BlueprintReadWrite, Setter="SetVerticalAlignment", Getter="GetVerticalAlignment", Category="Justify", meta=(AllowPrivateAccess="true"))
+	UPROPERTY(EditInstanceOnly, Setter="SetVerticalAlignment", Getter="GetVerticalAlignment", Category="Justify", meta=(AllowPrivateAccess="true"))
 	EAvaJustifyVertical VerticalAlignment;
 
-	UPROPERTY(EditAnywhere, BlueprintReadWrite, Setter="SetDepthAlignment", Getter="GetDepthAlignment", Category="Justify", meta=(AllowPrivateAccess="true"))
+	UPROPERTY(EditInstanceOnly, Setter="SetDepthAlignment", Getter="GetDepthAlignment", Category="Justify", meta=(AllowPrivateAccess="true"))
 	EAvaJustifyDepth DepthAlignment;
 
-	UPROPERTY(EditAnywhere, BlueprintReadWrite, Setter="SetHorizontalAnchor", Getter="GetHorizontalAnchor", Interp, Category="Justify", meta=(EditCondition="HorizontalAlignment != EAvaJustifyHorizontal::None", EditConditionHides, AllowPrivateAccess="true"))
+	UPROPERTY(EditInstanceOnly, Setter="SetHorizontalAnchor", Getter="GetHorizontalAnchor", Interp, Category="Justify", meta=(EditCondition="HorizontalAlignment != EAvaJustifyHorizontal::None", EditConditionHides, AllowPrivateAccess="true"))
 	float HorizontalAnchor;
 
-	UPROPERTY(EditAnywhere, BlueprintReadWrite, Setter="SetVerticalAnchor", Getter="GetVerticalAnchor", Interp, Category="Justify", meta=(EditCondition="VerticalAlignment != EAvaJustifyVertical::None", EditConditionHides, AllowPrivateAccess="true"))
+	UPROPERTY(EditInstanceOnly, Setter="SetVerticalAnchor", Getter="GetVerticalAnchor", Interp, Category="Justify", meta=(EditCondition="VerticalAlignment != EAvaJustifyVertical::None", EditConditionHides, AllowPrivateAccess="true"))
 	float VerticalAnchor;
 
-	UPROPERTY(EditAnywhere, BlueprintReadWrite, Setter="SetDepthAnchor", Getter="GetDepthAnchor", Interp, Category="Justify", meta=(EditCondition = "DepthAlignment != EAvaJustifyDepth::None", EditConditionHides, AllowPrivateAccess="true"))
+	UPROPERTY(EditInstanceOnly, Setter="SetDepthAnchor", Getter="GetDepthAnchor", Interp, Category="Justify", meta=(EditCondition = "DepthAlignment != EAvaJustifyDepth::None", EditConditionHides, AllowPrivateAccess="true"))
 	float DepthAnchor;
 
 private:
diff --git a/Engine/Plugins/Experimental/Avalanche/Source/AvalancheModifiers/Public/Modifiers/AvaLookAtModifier.h b/Engine/Plugins/Experimental/Avalanche/Source/AvalancheModifiers/Public/Modifiers/AvaLookAtModifier.h
--- a/Engine/Plugins/Experimental/Avalanche/Source/AvalancheModifiers/Public/Modifiers/AvaLookAtModifier.h
+++ b/Engine/Plugins/Experimental/Avalanche/Source/AvalancheModifiers/Public/Modifiers/AvaLookAtModifier.h
@@ -19,32 +19,32 @@ class UAvaLookAtModifier : public UAvaAttachmentBaseModifier
 	GENERATED_BODY()
 
 public:
-	UFUNCTION(BlueprintCallable, Category = "Motion Design|Modifiers|LookAt")
+	UFUNCTION(BlueprintCallable, Category="Motion Design|Modifiers|LookAt")
 	AVALANCHEMODIFIERS_API void SetReferenceActor(const FAvaSceneTreeActor& InReferenceActor);
 
-	UFUNCTION(BlueprintPure, Category = "Motion Design|Modifiers|LookAt")
+	UFUNCTION(BlueprintPure, Category="Motion Design|Modifiers|LookAt")
 	const FAvaSceneTreeActor& GetReferenceActor() const
 	{
 		return ReferenceActor;
 	}
 
 	/** Sets the axis that will point towards the reference actor. */
-	UFUNCTION(BlueprintCallable, Category = "Motion Design|Modifiers|LookAt")
-	AVALANCHEMODIFIERS_API void SetAxis(const EAvaAxis NewAxis);
+	UFUNCTION(BlueprintCallable, Category="Motion Design|Modifiers|LookAt")
+	AVALANCHEMODIFIERS_API void SetAxis(const EAvaAxis InAxis);
 
 	/** Returns the axis that will point towards t he reference actor. */
-	UFUNCTION(BlueprintPure, Category = "Motion Design|Modifiers|LookAt")
+	UFUNCTION(BlueprintPure, Category="Motion Design|Modifiers|LookAt")
 	EAvaAxis GetAxis() const
 	{
 		return Axis;
 	}
 
 	/** Sets the look-at direction to be flipped. */
-	UFUNCTION(BlueprintCallable, Category = "Motion Design|Modifiers|LookAt")
-	AVALANCHEMODIFIERS_API void SetFlipAxis(const bool bNewFlipAxis);
+	UFUNCTION(BlueprintCallable, Category="Motion Design|Modifiers|LookAt")
+	AVALANCHEMODIFIERS_API void SetFlipAxis(const bool bInFlipAxis);
 
 	/** Returns true if flipping the look-at rotation axis. */
-	UFUNCTION(BlueprintPure, Category = "Motion Design|Modifiers|LookAt")
+	UFUNCTION(BlueprintPure, Category="Motion Design|Modifiers|LookAt")
 	bool GetFlipAxis() const
 	{
 		return bFlipAxis;
@@ -78,7 +78,7 @@ protected:
 
 	void OnReferenceActorChanged();
 
-	UPROPERTY(EditAnywhere, BlueprintReadWrite, Setter="SetReferenceActor", Getter="GetReferenceActor", Category="LookAt", meta=(ShowOnlyInnerProperties, AllowPrivateAccess="true"))
+	UPROPERTY(EditInstanceOnly, Setter="SetReferenceActor", Getter="GetReferenceActor", Category="LookAt", meta=(ShowOnlyInnerProperties, AllowPrivateAccess="true"))
 	FAvaSceneTreeActor ReferenceActor;
 
 	/** The actor to look at. */
@@ -86,11 +86,11 @@ protected:
 	TWeakObjectPtr<AActor> ReferenceActorWeak;
 
 	/** The axis that will point towards the reference actor. */
-	UPROPERTY(EditAnywhere, BlueprintReadWrite, Setter="SetAxis", Getter="GetAxis", Category="LookAt", meta=(AllowPrivateAccess="true"))
+	UPROPERTY(EditInstanceOnly, Setter="SetAxis", Getter="GetAxis", Category="LookAt", meta=(AllowPrivateAccess="true"))
 	EAvaAxis Axis;
 
 	/** If true, will flip the look-at direction. */
-	UPROPERTY(EditAnywhere, BlueprintReadWrite, Setter="SetFlipAxis", Getter="GetFlipAxis", Category="LookAt", meta=(AllowPrivateAccess="true"))
+	UPROPERTY(EditInstanceOnly, Setter="SetFlipAxis", Getter="GetFlipAxis", Category="LookAt", meta=(AllowPrivateAccess="true"))
 	bool bFlipAxis = false;
 
 	UPROPERTY()
diff --git a/Engine/Plugins/Experimental/Avalanche/Source/AvalancheModifiers/Public/Modifiers/AvaMaterialParameterModifier.h b/Engine/Plugins/Experimental/Avalanche/Source/AvalancheModifiers/Public/Modifiers/AvaMaterialParameterModifier.h
--- a/Engine/Plugins/Experimental/Avalanche/Source/AvalancheModifiers/Public/Modifiers/AvaMaterialParameterModifier.h
+++ b/Engine/Plugins/Experimental/Avalanche/Source/AvalancheModifiers/Public/Modifiers/AvaMaterialParameterModifier.h
@@ -14,13 +14,13 @@ struct FAvaMaterialParameterMap
 {
 	GENERATED_BODY()
 
-	UPROPERTY(EditAnywhere, BlueprintReadWrite, Category = "MaterialParameter")
+	UPROPERTY(EditInstanceOnly, BlueprintReadWrite, Category="MaterialParameter")
 	TMap<FName, float> ScalarParameters;
 
-	UPROPERTY(EditAnywhere, BlueprintReadWrite, Category = "MaterialParameter")
+	UPROPERTY(EditInstanceOnly, BlueprintReadWrite, Category="MaterialParameter")
 	TMap<FName, FLinearColor> VectorParameters;
 
-	UPROPERTY(EditAnywhere, BlueprintReadWrite, Category = "MaterialParameter")
+	UPROPERTY(EditInstanceOnly, BlueprintReadWrite, Category="MaterialParameter")
 	TMap<FName, TObjectPtr<UTexture>> TextureParameters;
 
 	/** Matches the input parameter key map and removes all unused keys, does not touch current values */
@@ -42,13 +42,19 @@ class UAvaMaterialParameterModifier : public UAvaArrangeBaseModifier
 public:
 	UAvaMaterialParameterModifier();
 
+	UFUNCTION(BlueprintCallable, Category="Motion Design|Modifiers|MaterialParameter")
 	AVALANCHEMODIFIERS_API void SetMaterialParameters(const FAvaMaterialParameterMap& InParameterMap);
+
+	UFUNCTION(BlueprintPure, Category="Motion Design|Modifiers|MaterialParameter")
 	const FAvaMaterialParameterMap& GetMaterialParameters() const
 	{
 		return MaterialParameters;
 	}
 
+	UFUNCTION(BlueprintCallable, Category="Motion Design|Modifiers|MaterialParameter")
 	AVALANCHEMODIFIERS_API void SetUpdateChildren(bool bInUpdateChildren);
+
+	UFUNCTION(BlueprintPure, Category="Motion Design|Modifiers|MaterialParameter")
 	bool GetUpdateChildren() const
 	{
 		return bUpdateChildren;
@@ -99,9 +105,8 @@ protected:
 	/** Retrieves all Material Designer Instance from a primitive component */
 	TSet<UMaterialInstanceDynamic*> GetComponentDynamicMaterials(const UPrimitiveComponent* InComponent) const;
 
-	/** Which parameters should we set on the Material Designer Instance,
-	 * use EditCondition="bShowMaterialParameters && bShowMaterialParameters" otherwise edit inline boolean appear in details when it should not */
-	UPROPERTY(EditAnywhere, BlueprintReadWrite, Setter="SetMaterialParameters", Getter="GetMaterialParameters", Category="MaterialParameter", meta=(EditCondition="bShowMaterialParameters && bShowMaterialParameters", EditConditionHides, AllowPrivateAccess="true"))
+	/** Which parameters should we set on the Material Designer Instance */
+	UPROPERTY(EditInstanceOnly, Setter="SetMaterialParameters", Getter="GetMaterialParameters", Category="MaterialParameter", meta=(HideEditConditionToggle, EditCondition="bShowMaterialParameters", EditConditionHides, AllowPrivateAccess="true"))
 	FAvaMaterialParameterMap MaterialParameters;
 
 	/** Used to restore Material Designer Instance parameters to their original state */
@@ -113,7 +118,7 @@ protected:
 	TSubclassOf<UMaterialInstanceDynamic> MaterialClass;
 
 	/** Will also look into attached children actors */
-	UPROPERTY(EditAnywhere, BlueprintReadWrite, Setter="SetUpdateChildren", Getter="GetUpdateChildren", Category="MaterialParameter", meta=(AllowPrivateAccess="true"))
+	UPROPERTY(EditInstanceOnly, Setter="SetUpdateChildren", Getter="GetUpdateChildren", Category="MaterialParameter", meta=(AllowPrivateAccess="true"))
 	bool bUpdateChildren = true;
 
 #if WITH_EDITORONLY_DATA
diff --git a/Engine/Plugins/Experimental/Avalanche/Source/AvalancheModifiers/Public/Modifiers/AvaMirrorModifier.h b/Engine/Plugins/Experimental/Avalanche/Source/AvalancheModifiers/Public/Modifiers/AvaMirrorModifier.h
--- a/Engine/Plugins/Experimental/Avalanche/Source/AvalancheModifiers/Public/Modifiers/AvaMirrorModifier.h
+++ b/Engine/Plugins/Experimental/Avalanche/Source/AvalancheModifiers/Public/Modifiers/AvaMirrorModifier.h
@@ -14,31 +14,46 @@ class UAvaMirrorModifier : public UAvaGeometryBaseModifier
 	GENERATED_BODY()
 
 public:
+	UFUNCTION(BlueprintCallable, Category="Motion Design|Modifiers|Mirror")
 	AVALANCHEMODIFIERS_API void SetMirrorFramePosition(const FVector& InMirrorFramePosition);
+
+	UFUNCTION(BlueprintPure, Category="Motion Design|Modifiers|Mirror")
 	const FVector& GetMirrorFramePosition() const
 	{
 		return MirrorFramePosition;
 	}
 
+	UFUNCTION(BlueprintCallable, Category="Motion Design|Modifiers|Mirror")
 	AVALANCHEMODIFIERS_API void SetMirrorFrameRotation(const FRotator& InMirrorFrameRotation);
+
+	UFUNCTION(BlueprintPure, Category="Motion Design|Modifiers|Mirror")
 	const FRotator& GetMirrorFrameRotation() const
 	{
 		return MirrorFrameRotation;
 	}
 
+	UFUNCTION(BlueprintCallable, Category="Motion Design|Modifiers|Mirror")
 	AVALANCHEMODIFIERS_API void SetApplyPlaneCut(bool bInApplyPlaneCut);
+
+	UFUNCTION(BlueprintPure, Category="Motion Design|Modifiers|Mirror")
 	bool GetApplyPlaneCut() const
 	{
 		return bApplyPlaneCut;
 	}
 
+	UFUNCTION(BlueprintCallable, Category="Motion Design|Modifiers|Mirror")
 	AVALANCHEMODIFIERS_API void SetFlipCutSide(bool bInFlipCutSide);
+
+	UFUNCTION(BlueprintPure, Category="Motion Design|Modifiers|Mirror")
 	bool GetFlipCutSide() const
 	{
 		return bFlipCutSide;
 	}
 
+	UFUNCTION(BlueprintCallable, Category="Motion Design|Modifiers|Mirror")
 	AVALANCHEMODIFIERS_API void SetWeldAlongPlane(bool bInWeldAlongPlane);
+
+	UFUNCTION(BlueprintPure, Category="Motion Design|Modifiers|Mirror")
 	bool GetWeldAlongPlane() const
 	{
 		return bWeldAlongPlane;
@@ -70,23 +85,23 @@ protected:
 	void OnShowMirrorFrameChanged();
 #endif
 
-    UPROPERTY(EditInstanceOnly, BlueprintReadWrite, Setter="SetMirrorFramePosition", Getter="GetMirrorFramePosition", Category="Mirror", meta=(AllowPrivateAccess="true"))
+    UPROPERTY(EditInstanceOnly, Setter="SetMirrorFramePosition", Getter="GetMirrorFramePosition", Category="Mirror", meta=(AllowPrivateAccess="true"))
 	FVector MirrorFramePosition = FVector::ZeroVector;
 
-	UPROPERTY(EditInstanceOnly, BlueprintReadWrite, Setter="SetMirrorFrameRotation", Getter="GetMirrorFrameRotation", Category="Mirror", meta=(AllowPrivateAccess="true"))
+	UPROPERTY(EditInstanceOnly, Setter="SetMirrorFrameRotation", Getter="GetMirrorFrameRotation", Category="Mirror", meta=(AllowPrivateAccess="true"))
 	FRotator MirrorFrameRotation = FRotator(0, 0, 90);
 
-	UPROPERTY(EditInstanceOnly, BlueprintReadWrite, Setter="SetApplyPlaneCut", Getter="GetApplyPlaneCut", Category="Mirror", meta=(AllowPrivateAccess="true"))
+	UPROPERTY(EditInstanceOnly, Setter="SetApplyPlaneCut", Getter="GetApplyPlaneCut", Category="Mirror", meta=(AllowPrivateAccess="true"))
 	bool bApplyPlaneCut = true;
 
-	UPROPERTY(EditInstanceOnly, BlueprintReadWrite, Setter="SetFlipCutSide", Getter="GetFlipCutSide", Category="Mirror", meta=(AllowPrivateAccess="true"))
+	UPROPERTY(EditInstanceOnly, Setter="SetFlipCutSide", Getter="GetFlipCutSide", Category="Mirror", meta=(AllowPrivateAccess="true"))
 	bool bFlipCutSide = false;
 
-	UPROPERTY(EditInstanceOnly, BlueprintReadWrite, Setter="SetWeldAlongPlane", Getter="GetWeldAlongPlane", Category="Mirror", meta=(AllowPrivateAccess="true"))
+	UPROPERTY(EditInstanceOnly, Setter="SetWeldAlongPlane", Getter="GetWeldAlongPlane", Category="Mirror", meta=(AllowPrivateAccess="true"))
 	bool bWeldAlongPlane = true;
 
 #if WITH_EDITORONLY_DATA
-	UPROPERTY(EditInstanceOnly, BlueprintReadWrite, Category="Mirror", meta=(AllowPrivateAccess="true"))
+	UPROPERTY(EditInstanceOnly, Category="Mirror", meta=(AllowPrivateAccess="true"))
 	bool bShowMirrorFrame = false;
 
 private:
diff --git a/Engine/Plugins/Experimental/Avalanche/Source/AvalancheModifiers/Public/Modifiers/AvaNormalModifier.h b/Engine/Plugins/Experimental/Avalanche/Source/AvalancheModifiers/Public/Modifiers/AvaNormalModifier.h
--- a/Engine/Plugins/Experimental/Avalanche/Source/AvalancheModifiers/Public/Modifiers/AvaNormalModifier.h
+++ b/Engine/Plugins/Experimental/Avalanche/Source/AvalancheModifiers/Public/Modifiers/AvaNormalModifier.h
@@ -26,40 +26,59 @@ class UAvaNormalModifier : public UAvaGeometryBaseModifier
 	GENERATED_BODY()
 
 public:
+	UFUNCTION(BlueprintCallable, Category="Motion Design|Modifiers|Normal")
 	AVALANCHEMODIFIERS_API void SetAngleWeighted(bool bInAngleWeighted);
+
+	UFUNCTION(BlueprintPure, Category="Motion Design|Modifiers|Normal")
 	bool GetAngleWeighted() const
 	{
 		return bAngleWeighted;
 	}
 
+	UFUNCTION(BlueprintCallable, Category="Motion Design|Modifiers|Normal")
 	AVALANCHEMODIFIERS_API void SetAreaWeighted(bool bInAreaWeighted);
+
+	UFUNCTION(BlueprintPure, Category="Motion Design|Modifiers|Normal")
 	bool GetAreaWeighted() const
 	{
 		return bAreaWeighted;
 	}
 
+	UFUNCTION(BlueprintCallable, Category="Motion Design|Modifiers|Normal")
 	AVALANCHEMODIFIERS_API void SetInvert(bool bInInvert);
+
+	UFUNCTION(BlueprintPure, Category="Motion Design|Modifiers|Normal")
 	bool GetInvert() const
 	{
 		return bInvert;
 	}
 
+	UFUNCTION(BlueprintCallable, Category="Motion Design|Modifiers|Normal")
 	AVALANCHEMODIFIERS_API void SetSplitMethod(EAvaNormalModifierSplitMethod InSplitMethod);
+
+	UFUNCTION(BlueprintPure, Category="Motion Design|Modifiers|Normal")
 	EAvaNormalModifierSplitMethod GetSplitMethod() const
 	{
 		return SplitMethod;
 	}
 
+	UFUNCTION(BlueprintCallable, Category="Motion Design|Modifiers|Normal")
 	AVALANCHEMODIFIERS_API void SetAngleThreshold(float InAngleThreshold);
+
+	UFUNCTION(BlueprintPure, Category="Motion Design|Modifiers|Normal")
 	float GetAngleThreshold() const
 	{
 		return AngleThreshold;
 	}
 
+	UFUNCTION(BlueprintCallable, Category="Motion Design|Modifiers|Normal")
 	AVALANCHEMODIFIERS_API void SetPolyGroupLayerIdx(int32 InPolyGroupLayer);
+
+	UFUNCTION(BlueprintPure, Category="Motion Design|Modifiers|Normal")
 	AVALANCHEMODIFIERS_API int32 GetPolyGroupLayerIdx() const;
 
-	AVALANCHEMODIFIERS_API void SetPolyGroupLayer(FString& InString);
+	AVALANCHEMODIFIERS_API void SetPolyGroupLayer(const FString& InPolyGroupLayer);
+
 	FString GetPolyGroupLayer() const
 	{
 		return PolyGroupLayer;
@@ -89,26 +108,26 @@ protected:
 	TArray<FString> GetPolyGroupLayers() const;
 
 	/** Recompute normals and weight them by angle */
-	UPROPERTY(EditInstanceOnly, BlueprintReadWrite, Setter="SetAngleWeighted", Getter="GetAngleWeighted", Category="Normal", meta=(AllowPrivateAccess="true"))
+	UPROPERTY(EditInstanceOnly, Setter="SetAngleWeighted", Getter="GetAngleWeighted", Category="Normal", meta=(AllowPrivateAccess="true"))
 	bool bAngleWeighted = true;
 
 	/** Recompute normals and weight them by area */
-	UPROPERTY(EditInstanceOnly, BlueprintReadWrite, Setter="SetAreaWeighted", Getter="GetAreaWeighted", Category="Normal", meta=(AllowPrivateAccess="true"))
+	UPROPERTY(EditInstanceOnly, Setter="SetAreaWeighted", Getter="GetAreaWeighted", Category="Normal", meta=(AllowPrivateAccess="true"))
 	bool bAreaWeighted = true;
 
 	/** Recompute normals and invert normals and triangles */
-	UPROPERTY(EditInstanceOnly, BlueprintReadWrite, Setter="SetInvert", Getter="GetInvert", Category="Normal", meta=(AllowPrivateAccess="true"))
+	UPROPERTY(EditInstanceOnly, Setter="SetInvert", Getter="GetInvert", Category="Normal", meta=(AllowPrivateAccess="true"))
 	bool bInvert = false;
 
 	/** Recompute normals and use a split method */
-	UPROPERTY(EditInstanceOnly, BlueprintReadWrite, Setter="SetSplitMethod", Getter="GetSplitMethod", Category="Normal", meta=(AllowPrivateAccess="true"))
+	UPROPERTY(EditInstanceOnly, Setter="SetSplitMethod", Getter="GetSplitMethod", Category="Normal", meta=(AllowPrivateAccess="true"))
 	EAvaNormalModifierSplitMethod SplitMethod = EAvaNormalModifierSplitMethod::Threshold;
 
 	/** Angle to compare and split normal when threshold method is chosen */
-	UPROPERTY(EditInstanceOnly, BlueprintReadWrite, Setter="SetAngleThreshold", Getter="GetAngleThreshold", Category="Normal", meta=(ClampMin="0.0", ClampMax="180.0", EditCondition="SplitMethod == EAvaNormalModifierSplitMethod::Threshold", AllowPrivateAccess="true"))
+	UPROPERTY(EditInstanceOnly, Setter="SetAngleThreshold", Getter="GetAngleThreshold", Category="Normal", meta=(ClampMin="0.0", ClampMax="180.0", EditCondition="SplitMethod == EAvaNormalModifierSplitMethod::Threshold", AllowPrivateAccess="true"))
 	float AngleThreshold = 60.f;
 
 	/** PolyGroup to use to split normal from when PolyGroup method is chosen */
-	UPROPERTY(EditInstanceOnly, BlueprintReadWrite, Setter="SetPolyGroupLayer", Getter="GetPolyGroupLayer", Category="Normal", meta=(GetOptions="GetPolyGroupLayers", EditCondition="SplitMethod == EAvaNormalModifierSplitMethod::PolyGroup", AllowPrivateAccess="true"))
+	UPROPERTY(EditInstanceOnly, Setter="SetPolyGroupLayer", Getter="GetPolyGroupLayer", Category="Normal", meta=(GetOptions="GetPolyGroupLayers", EditCondition="SplitMethod == EAvaNormalModifierSplitMethod::PolyGroup", AllowPrivateAccess="true"))
 	FString PolyGroupLayer = TEXT("None");
 };
diff --git a/Engine/Plugins/Experimental/Avalanche/Source/AvalancheModifiers/Public/Modifiers/AvaOutlineModifier.h b/Engine/Plugins/Experimental/Avalanche/Source/AvalancheModifiers/Public/Modifiers/AvaOutlineModifier.h
--- a/Engine/Plugins/Experimental/Avalanche/Source/AvalancheModifiers/Public/Modifiers/AvaOutlineModifier.h
+++ b/Engine/Plugins/Experimental/Avalanche/Source/AvalancheModifiers/Public/Modifiers/AvaOutlineModifier.h
@@ -19,19 +19,28 @@ class UAvaOutlineModifier : public UAvaGeometryBaseModifier
 	GENERATED_BODY()
 
 public:
+	UFUNCTION(BlueprintCallable, Category="Motion Design|Modifiers|Outline")
 	AVALANCHEMODIFIERS_API void SetMode(EAvaOutlineMode InMode);
+
+	UFUNCTION(BlueprintPure, Category="Motion Design|Modifiers|Outline")
 	EAvaOutlineMode GetMode() const
 	{
 		return Mode;
 	}
 
+	UFUNCTION(BlueprintCallable, Category="Motion Design|Modifiers|Outline")
 	AVALANCHEMODIFIERS_API void SetDistance(float InDistance);
+
+	UFUNCTION(BlueprintPure, Category="Motion Design|Modifiers|Outline")
 	float GetDistance() const
 	{
 		return Distance;
 	}
 
+	UFUNCTION(BlueprintCallable, Category="Motion Design|Modifiers|Outline")
 	AVALANCHEMODIFIERS_API void SetRemoveInside(bool bInRemoveInside);
+
+	UFUNCTION(BlueprintPure, Category="Motion Design|Modifiers|Outline")
 	bool GetRemoveInside() const
 	{
 		return bRemoveInside;
@@ -56,14 +65,14 @@ protected:
 	float GetMaxInsetDistance() const;
 
 	/** Set the mode like inset or outset */
-	UPROPERTY(EditInstanceOnly, BlueprintReadWrite, Setter="SetMode", Getter="GetMode", Category="Outline", meta=(AllowPrivateAccess="true"))
+	UPROPERTY(EditInstanceOnly, Setter="SetMode", Getter="GetMode", Category="Outline", meta=(AllowPrivateAccess="true"))
 	EAvaOutlineMode Mode = EAvaOutlineMode::Outset;
 
 	/** Set the distance for the outline */
-	UPROPERTY(EditInstanceOnly, BlueprintReadWrite, Setter="SetDistance", Getter="GetDistance", Category="Outline", meta=(ClampMin="0", AllowPrivateAccess="true"))
+	UPROPERTY(EditInstanceOnly, Setter="SetDistance", Getter="GetDistance", Category="Outline", meta=(ClampMin="0", AllowPrivateAccess="true"))
 	float Distance = 10.f;
 
 	/** Remove the inside part and create a hole in the shape */
-	UPROPERTY(EditInstanceOnly, BlueprintReadWrite, Setter="SetRemoveInside", Getter="GetRemoveInside", Category="Outline", meta=(AllowPrivateAccess="true"))
+	UPROPERTY(EditInstanceOnly, Setter="SetRemoveInside", Getter="GetRemoveInside", Category="Outline", meta=(AllowPrivateAccess="true"))
 	bool bRemoveInside = true;
 };
diff --git a/Engine/Plugins/Experimental/Avalanche/Source/AvalancheModifiers/Public/Modifiers/AvaPatternModifier.h b/Engine/Plugins/Experimental/Avalanche/Source/AvalancheModifiers/Public/Modifiers/AvaPatternModifier.h
--- a/Engine/Plugins/Experimental/Avalanche/Source/AvalancheModifiers/Public/Modifiers/AvaPatternModifier.h
+++ b/Engine/Plugins/Experimental/Avalanche/Source/AvalancheModifiers/Public/Modifiers/AvaPatternModifier.h
@@ -160,25 +160,37 @@ class UAvaPatternModifier : public UAvaGeometryBaseModifier
 	GENERATED_BODY()
 
 public:
+	UFUNCTION(BlueprintCallable, Category="Motion Design|Modifiers|Pattern")
 	AVALANCHEMODIFIERS_API void SetLayout(EAvaPatternModifierLayout InLayout);
+
+	UFUNCTION(BlueprintPure, Category="Motion Design|Modifiers|Pattern")
 	EAvaPatternModifierLayout GetLayout() const
 	{
 		return Layout;
 	}
 
+	UFUNCTION(BlueprintCallable, Category="Motion Design|Modifiers|Pattern")
 	AVALANCHEMODIFIERS_API void SetLineLayoutOptions(const FAvaPatternModifierLineLayoutOptions& InOptions);
+
+	UFUNCTION(BlueprintPure, Category="Motion Design|Modifiers|Pattern")
 	const FAvaPatternModifierLineLayoutOptions& GetLineLayoutOptions() const
 	{
 		return LineLayoutOptions;
 	}
 
+	UFUNCTION(BlueprintCallable, Category="Motion Design|Modifiers|Pattern")
 	AVALANCHEMODIFIERS_API void SetGridLayoutOptions(const FAvaPatternModifierGridLayoutOptions& InOptions);
+
+	UFUNCTION(BlueprintPure, Category="Motion Design|Modifiers|Pattern")
 	const FAvaPatternModifierGridLayoutOptions& GetGridLayoutOptions() const
 	{
 		return GridLayoutOptions;
 	}
 
+	UFUNCTION(BlueprintCallable, Category="Motion Design|Modifiers|Pattern")
 	AVALANCHEMODIFIERS_API void SetCircleLayoutOptions(const FAvaPatternModifierCircleLayoutOptions& InOptions);
+
+	UFUNCTION(BlueprintPure, Category="Motion Design|Modifiers|Pattern")
 	const FAvaPatternModifierCircleLayoutOptions& GetCircleLayoutOptions() const
 	{
 		return CircleLayoutOptions;
@@ -206,19 +218,19 @@ protected:
 	UE::Geometry::FTransformSRT3d GetGridLayoutRowTransformChange() const;
 	UE::Geometry::FTransformSRT3d GetCircleLayoutTransformChange(int32 Idx) const;
 
-	UPROPERTY(EditInstanceOnly, BlueprintReadWrite, Setter="SetLayout", Getter="GetLayout", Category="Pattern", meta=(AllowPrivateAccess="true"))
+	UPROPERTY(EditInstanceOnly, Setter="SetLayout", Getter="GetLayout", Category="Pattern", meta=(AllowPrivateAccess="true"))
 	EAvaPatternModifierLayout Layout = EAvaPatternModifierLayout::Line;
 
 	/** Line layout options */
-	UPROPERTY(EditInstanceOnly, BlueprintReadWrite, Setter="SetLineLayoutOptions", Getter="GetLineLayoutOptions", Category="Pattern", meta=(EditCondition="Layout == EAvaPatternModifierLayout::Line", EditConditionHides, AllowPrivateAccess="true"))
+	UPROPERTY(EditInstanceOnly, Setter="SetLineLayoutOptions", Getter="GetLineLayoutOptions", Category="Pattern", meta=(EditCondition="Layout == EAvaPatternModifierLayout::Line", EditConditionHides, AllowPrivateAccess="true"))
 	FAvaPatternModifierLineLayoutOptions LineLayoutOptions;
 
 	/** Grid layout options */
-	UPROPERTY(EditInstanceOnly, BlueprintReadWrite, Setter="SetGridLayoutOptions", Getter="GetGridLayoutOptions", Category="Pattern", meta=(EditCondition="Layout == EAvaPatternModifierLayout::Grid", EditConditionHides, AllowPrivateAccess="true"))
+	UPROPERTY(EditInstanceOnly, Setter="SetGridLayoutOptions", Getter="GetGridLayoutOptions", Category="Pattern", meta=(EditCondition="Layout == EAvaPatternModifierLayout::Grid", EditConditionHides, AllowPrivateAccess="true"))
 	FAvaPatternModifierGridLayoutOptions GridLayoutOptions;
 
 	/** Circle layout options */
-	UPROPERTY(EditInstanceOnly, BlueprintReadWrite, Setter="SetCircleLayoutOptions", Getter="GetCircleLayoutOptions", Category="Pattern", meta=(EditCondition="Layout == EAvaPatternModifierLayout::Circle", EditConditionHides, AllowPrivateAccess="true"))
+	UPROPERTY(EditInstanceOnly, Setter="SetCircleLayoutOptions", Getter="GetCircleLayoutOptions", Category="Pattern", meta=(EditCondition="Layout == EAvaPatternModifierLayout::Circle", EditConditionHides, AllowPrivateAccess="true"))
 	FAvaPatternModifierCircleLayoutOptions CircleLayoutOptions;
 
 	UPROPERTY(Transient, DuplicateTransient, TextExportTransient)
diff --git a/Engine/Plugins/Experimental/Avalanche/Source/AvalancheModifiers/Public/Modifiers/AvaPlaneCutModifier.h b/Engine/Plugins/Experimental/Avalanche/Source/AvalancheModifiers/Public/Modifiers/AvaPlaneCutModifier.h
--- a/Engine/Plugins/Experimental/Avalanche/Source/AvalancheModifiers/Public/Modifiers/AvaPlaneCutModifier.h
+++ b/Engine/Plugins/Experimental/Avalanche/Source/AvalancheModifiers/Public/Modifiers/AvaPlaneCutModifier.h
@@ -13,25 +13,37 @@ class UAvaPlaneCutModifier : public UAvaGeometryBaseModifier
 	GENERATED_BODY()
 
 public:
+	UFUNCTION(BlueprintCallable, Category="Motion Design|Modifiers|PlaneCut")
 	AVALANCHEMODIFIERS_API void SetPlaneOrigin(float InOrigin);
+
+	UFUNCTION(BlueprintPure, Category="Motion Design|Modifiers|PlaneCut")
 	float GetPlaneOrigin() const
 	{
 		return PlaneOrigin;
 	}
 
+	UFUNCTION(BlueprintCallable, Category="Motion Design|Modifiers|PlaneCut")
 	AVALANCHEMODIFIERS_API void SetPlaneRotation(const FRotator& InRotation);
+
+	UFUNCTION(BlueprintPure, Category="Motion Design|Modifiers|PlaneCut")
 	const FRotator& GetPlaneRotation() const
 	{
 		return PlaneRotation;
 	}
 
+	UFUNCTION(BlueprintCallable, Category="Motion Design|Modifiers|PlaneCut")
 	AVALANCHEMODIFIERS_API void SetInvertCut(bool bInInvertCut);
+
+	UFUNCTION(BlueprintPure, Category="Motion Design|Modifiers|PlaneCut")
 	bool GetInvertCut() const
 	{
 		return bInvertCut;
 	}
 
+	UFUNCTION(BlueprintCallable, Category="Motion Design|Modifiers|PlaneCut")
 	AVALANCHEMODIFIERS_API void SetFillHoles(bool bInFillHoles);
+
+	UFUNCTION(BlueprintPure, Category="Motion Design|Modifiers|PlaneCut")
 	bool GetFillHoles() const
 	{
 		return bFillHoles;
@@ -67,20 +79,20 @@ protected:
 	void UpdatePreviewComponent();
 #endif
 
-	UPROPERTY(EditInstanceOnly, BlueprintReadWrite, Setter="SetPlaneOrigin", Getter="GetPlaneOrigin", Category="PlaneCut", meta=(Delta="0.5", LinearDeltaSensitivity="1", AllowPrivateAccess="true"))
+	UPROPERTY(EditInstanceOnly, Setter="SetPlaneOrigin", Getter="GetPlaneOrigin", Category="PlaneCut", meta=(Delta="0.5", LinearDeltaSensitivity="1", AllowPrivateAccess="true"))
 	float PlaneOrigin = 0.f;
 
-	UPROPERTY(EditInstanceOnly, BlueprintReadWrite, Setter="SetPlaneRotation", Getter="GetPlaneRotation", Category="PlaneCut", meta=(AllowPrivateAccess="true"))
+	UPROPERTY(EditInstanceOnly, Setter="SetPlaneRotation", Getter="GetPlaneRotation", Category="PlaneCut", meta=(AllowPrivateAccess="true"))
 	FRotator PlaneRotation = FRotator(0, 0, 90);
 
-	UPROPERTY(EditInstanceOnly, BlueprintReadWrite, Setter="SetInvertCut", Getter="GetInvertCut", Category="PlaneCut", meta=(AllowPrivateAccess="true"))
+	UPROPERTY(EditInstanceOnly, Setter="SetInvertCut", Getter="GetInvertCut", Category="PlaneCut", meta=(AllowPrivateAccess="true"))
 	bool bInvertCut = false;
 
-	UPROPERTY(EditInstanceOnly, BlueprintReadWrite, Setter="SetFillHoles", Getter="GetFillHoles", Category="PlaneCut", meta=(AllowPrivateAccess="true"))
+	UPROPERTY(EditInstanceOnly, Setter="SetFillHoles", Getter="GetFillHoles", Category="PlaneCut", meta=(AllowPrivateAccess="true"))
 	bool bFillHoles = true;
 
 #if WITH_EDITORONLY_DATA
-	UPROPERTY(EditInstanceOnly, BlueprintReadWrite, Category="PlaneCut", meta=(AllowPrivateAccess="true"))
+	UPROPERTY(EditInstanceOnly, Category="PlaneCut", meta=(AllowPrivateAccess="true"))
 	bool bUsePreview = false;
 
 private:
diff --git a/Engine/Plugins/Experimental/Avalanche/Source/AvalancheModifiers/Public/Modifiers/AvaRadialArrangeModifier.h b/Engine/Plugins/Experimental/Avalanche/Source/AvalancheModifiers/Public/Modifiers/AvaRadialArrangeModifier.h
--- a/Engine/Plugins/Experimental/Avalanche/Source/AvalancheModifiers/Public/Modifiers/AvaRadialArrangeModifier.h
+++ b/Engine/Plugins/Experimental/Avalanche/Source/AvalancheModifiers/Public/Modifiers/AvaRadialArrangeModifier.h
@@ -31,121 +31,121 @@ class UAvaRadialArrangeModifier : public UAvaArrangeBaseModifier
 
 public:
 	/** Sets the number of child elements to use in the arrangement. Children whose index is greater than or equal to this value will be hidden. */
-	UFUNCTION(BlueprintCallable, Category = "Motion Design|Modifiers|RadialArrange")
-	AVALANCHEMODIFIERS_API void SetCount(const int32 NewCount);
+	UFUNCTION(BlueprintCallable, Category="Motion Design|Modifiers|RadialArrange")
+	AVALANCHEMODIFIERS_API void SetCount(const int32 InCount);
 
 	/** Gets the number of child elements to use in the arrangement. Children whose index is greater than or equal to this value will be hidden. */
-	UFUNCTION(BlueprintPure, Category = "Motion Design|Modifiers|RadialArrange")
+	UFUNCTION(BlueprintPure, Category="Motion Design|Modifiers|RadialArrange")
 	int32 GetCount() const
 	{
 		return Count;
 	}
 
 	/** Sets the number of rings. */
-	UFUNCTION(BlueprintCallable, Category = "Motion Design|Modifiers|RadialArrange")
-	AVALANCHEMODIFIERS_API void SetRings(const int32 NewRings);
+	UFUNCTION(BlueprintCallable, Category="Motion Design|Modifiers|RadialArrange")
+	AVALANCHEMODIFIERS_API void SetRings(const int32 InRings);
 
 	/** Gets the number of rings. */
-	UFUNCTION(BlueprintPure, Category = "Motion Design|Modifiers|RadialArrange")
+	UFUNCTION(BlueprintPure, Category="Motion Design|Modifiers|RadialArrange")
 	int32 GetRings() const
 	{
 		return Rings;
 	}
 
 	/** Sets the radius from the center to the first inner ring. */
-	UFUNCTION(BlueprintCallable, Category = "Motion Design|Modifiers|RadialArrange")
-	AVALANCHEMODIFIERS_API void SetInnerRadius(const float NewInnerRadius);
+	UFUNCTION(BlueprintCallable, Category="Motion Design|Modifiers|RadialArrange")
+	AVALANCHEMODIFIERS_API void SetInnerRadius(const float InInnerRadius);
 
 	/** Gets the radius from the center to the first inner ring. */
-	UFUNCTION(BlueprintPure, Category = "Motion Design|Modifiers|RadialArrange")
+	UFUNCTION(BlueprintPure, Category="Motion Design|Modifiers|RadialArrange")
 	float GetInnerRadius() const
 	{
 		return InnerRadius;
 	}
 
 	/** Sets the radius from the center to the last outer ring. */
-	UFUNCTION(BlueprintCallable, Category = "Motion Design|Modifiers|RadialArrange")
-	AVALANCHEMODIFIERS_API void SetOuterRadius(const float NewOuterRadius);
+	UFUNCTION(BlueprintCallable, Category="Motion Design|Modifiers|RadialArrange")
+	AVALANCHEMODIFIERS_API void SetOuterRadius(const float InOuterRadius);
 
 	/** Gets the radius from the center to the last outer ring. */
-	UFUNCTION(BlueprintPure, Category = "Motion Design|Modifiers|RadialArrange")
+	UFUNCTION(BlueprintPure, Category="Motion Design|Modifiers|RadialArrange")
 	float GetOuterRadius() const
 	{
 		return OuterRadius;
 	}
 
 	/** Sets the start angle for the arrangement space and moving clockwise. 0 = Up, -90 = Left, 90 = Right */
-	UFUNCTION(BlueprintCallable, Category = "Motion Design|Modifiers|RadialArrange")
-	AVALANCHEMODIFIERS_API void SetStartAngle(const float NewStartAngle);
+	UFUNCTION(BlueprintCallable, Category="Motion Design|Modifiers|RadialArrange")
+	AVALANCHEMODIFIERS_API void SetStartAngle(const float InStartAngle);
 
 	/** Gets the start angle for the arrangement space and moving clockwise. 0 = Up, -90 = Left, 90 = Right */
-	UFUNCTION(BlueprintPure, Category = "Motion Design|Modifiers|RadialArrange")
+	UFUNCTION(BlueprintPure, Category="Motion Design|Modifiers|RadialArrange")
 	float GetStartAngle() const
 	{
 		return StartAngle;
 	}
 
 	/** Sets the end angle for the arrangement space. 0 = Up, -90 = Left, 90 = Right */
-	UFUNCTION(BlueprintCallable, Category = "Motion Design|Modifiers|RadialArrange")
-	AVALANCHEMODIFIERS_API void SetEndAngle(const float NewEndAngle);
+	UFUNCTION(BlueprintCallable, Category="Motion Design|Modifiers|RadialArrange")
+	AVALANCHEMODIFIERS_API void SetEndAngle(const float InEndAngle);
 
 	/** Gets the end angle for the arrangement space. 0 = Up, -90 = Left, 90 = Right */
-	UFUNCTION(BlueprintPure, Category = "Motion Design|Modifiers|RadialArrange")
+	UFUNCTION(BlueprintPure, Category="Motion Design|Modifiers|RadialArrange")
 	float GetEndAngle() const
 	{
 		return EndAngle;
 	}
 
 	/** Defines how to arrange the child elements around the center. */
-	UFUNCTION(BlueprintCallable, Category = "Motion Design|Modifiers|RadialArrange")
-	AVALANCHEMODIFIERS_API void SetArrangement(const EAvaRadialArrangement NewArrangement);
+	UFUNCTION(BlueprintCallable, Category="Motion Design|Modifiers|RadialArrange")
+	AVALANCHEMODIFIERS_API void SetArrangement(const EAvaRadialArrangement InArrangement);
 
 	/** Defines how to arrange the child elements around the center. */
-	UFUNCTION(BlueprintPure, Category = "Motion Design|Modifiers|RadialArrange")
+	UFUNCTION(BlueprintPure, Category="Motion Design|Modifiers|RadialArrange")
 	EAvaRadialArrangement GetArrangement() const
 	{
 		return Arrangement;
 	}
 
 	/** If true, will arrange the child elements starting from the outer radius and moving to the inner radius. Has no effect if only using one ring. */
-	UFUNCTION(BlueprintCallable, Category = "Motion Design|Modifiers|RadialArrange")
-	AVALANCHEMODIFIERS_API void SetStartFromOuterRadius(const bool bMewStartFromOuterRadius);
+	UFUNCTION(BlueprintCallable, Category="Motion Design|Modifiers|RadialArrange")
+	AVALANCHEMODIFIERS_API void SetStartFromOuterRadius(const bool bInStartFromOuterRadius);
 
 	/** If true, will arrange the child elements starting from the outer radius and moving to the inner radius. Has no effect if only using one ring. */
-	UFUNCTION(BlueprintPure, Category = "Motion Design|Modifiers|RadialArrange")
+	UFUNCTION(BlueprintPure, Category="Motion Design|Modifiers|RadialArrange")
 	bool GetStartFromOuterRadius() const
 	{
 		return bStartFromOuterRadius;
 	}
 
 	/** If true, will orient the selected axis torwards the center. */
-	UFUNCTION(BlueprintCallable, Category = "Motion Design|Modifiers|RadialArrange")
-	AVALANCHEMODIFIERS_API void SetOrient(const bool bNewOrient);
+	UFUNCTION(BlueprintCallable, Category="Motion Design|Modifiers|RadialArrange")
+	AVALANCHEMODIFIERS_API void SetOrient(const bool bInOrient);
 
 	/** If true, will orient the selected axis torwards the center. */
-	UFUNCTION(BlueprintPure, Category = "Motion Design|Modifiers|RadialArrange")
+	UFUNCTION(BlueprintPure, Category="Motion Design|Modifiers|RadialArrange")
 	bool GetOrient() const
 	{
 		return bOrient;
 	}
 
 	/** Sets the axis to look at the center. */
-	UFUNCTION(BlueprintCallable, Category = "Motion Design|Modifiers|RadialArrange")
-	AVALANCHEMODIFIERS_API void SetOrientAxis(const EAvaAxis NewOrientAxis);
+	UFUNCTION(BlueprintCallable, Category="Motion Design|Modifiers|RadialArrange")
+	AVALANCHEMODIFIERS_API void SetOrientAxis(const EAvaAxis InOrientAxis);
 
 	/** Gets the axis to look at the center. */
-	UFUNCTION(BlueprintPure, Category = "Motion Design|Modifiers|RadialArrange")
+	UFUNCTION(BlueprintPure, Category="Motion Design|Modifiers|RadialArrange")
 	EAvaAxis GetOrientAxis() const
 	{
 		return OrientAxis;
 	}
 
 	/** If true, will flip the center orientation to face outwards. */
-	UFUNCTION(BlueprintCallable, Category = "Motion Design|Modifiers|RadialArrange")
-	AVALANCHEMODIFIERS_API void SetFlipOrient(const bool bNewFlipOrient);
+	UFUNCTION(BlueprintCallable, Category="Motion Design|Modifiers|RadialArrange")
+	AVALANCHEMODIFIERS_API void SetFlipOrient(const bool bInFlipOrient);
 
 	/** If true, will flip the center orientation to face outwards. */
-	UFUNCTION(BlueprintPure, Category = "Motion Design|Modifiers|RadialArrange")
+	UFUNCTION(BlueprintPure, Category="Motion Design|Modifiers|RadialArrange")
 	bool GetFlipOrient() const
 	{
 		return bFlipOrient;
@@ -164,46 +164,46 @@ protected:
 	//~ End UActorModifierCoreBase
 
 	/** The number of child elements to limit in the arrangement, or -1 if unlimited. Children whose index is greater than or equal to this value will be hidden. */
-	UPROPERTY(EditAnywhere, BlueprintReadWrite, Setter="SetCount", Getter="GetCount", Category="RadialArrange", meta=(ClampMin="-1", UIMin="-1", AllowPrivateAccess="true"))
+	UPROPERTY(EditInstanceOnly, Setter="SetCount", Getter="GetCount", Category="RadialArrange", meta=(ClampMin="-1", UIMin="-1", AllowPrivateAccess="true"))
 	int32 Count = -1;
 
 	/** The number of rings. */
-	UPROPERTY(EditAnywhere, BlueprintReadWrite, Setter="SetRings", Getter="GetRings", Category="RadialArrange", meta=(ClampMin="1", UIMin="1", AllowPrivateAccess="true"))
+	UPROPERTY(EditInstanceOnly, Setter="SetRings", Getter="GetRings", Category="RadialArrange", meta=(ClampMin="1", UIMin="1", AllowPrivateAccess="true"))
 	int32 Rings = 1;
 
 	/** The radius from the center to the first inner ring. */
-	UPROPERTY(EditAnywhere, BlueprintReadWrite, Setter="SetInnerRadius", Getter="GetInnerRadius", Interp, Category="RadialArrange", meta=(ClampMin="0.0", UIMin="0.0", AllowPrivateAccess="true"))
+	UPROPERTY(EditInstanceOnly, Setter="SetInnerRadius", Getter="GetInnerRadius", Interp, Category="RadialArrange", meta=(ClampMin="0.0", UIMin="0.0", AllowPrivateAccess="true"))
 	float InnerRadius = 70.0f;
 
 	/** The radius from the center to the last outer ring. */
-	UPROPERTY(EditAnywhere, BlueprintReadWrite, Setter="SetOuterRadius", Getter="GetOuterRadius", Interp, Category="RadialArrange", meta=(EditCondition="Rings > 1", EditConditionHides, ClampMin="0.0", UIMin="0.0", AllowPrivateAccess="true"))
+	UPROPERTY(EditInstanceOnly, Setter="SetOuterRadius", Getter="GetOuterRadius", Interp, Category="RadialArrange", meta=(EditCondition="Rings > 1", EditConditionHides, ClampMin="0.0", UIMin="0.0", AllowPrivateAccess="true"))
 	float OuterRadius = 200.0f;
 
 	/** The start angle for the arrangement space and moving clockwise. 0 = Up, -90 = Left, 90 = Right */
-	UPROPERTY(EditAnywhere, BlueprintReadWrite, Setter="SetStartAngle", Getter="GetStartAngle", Interp, Category="RadialArrange", meta=(ClampMin="-180.0", UIMin="-180.0", ClampMax="180.0", UIMax="180.0", AllowPrivateAccess="true"))
+	UPROPERTY(EditInstanceOnly, Setter="SetStartAngle", Getter="GetStartAngle", Interp, Category="RadialArrange", meta=(ClampMin="-180.0", UIMin="-180.0", ClampMax="180.0", UIMax="180.0", AllowPrivateAccess="true"))
 	float StartAngle = -180.0f;
 
 	/** The end angle for the arrangement space. 0 = Up, -90 = Left, 90 = Right */
-	UPROPERTY(EditAnywhere, BlueprintReadWrite, Setter="SetEndAngle", Getter="GetEndAngle", Interp, Category="RadialArrange", meta=(ClampMin = "-180.0", UIMin = "-180.0", ClampMax = "180.0", UIMax = "180.0", AllowPrivateAccess="true"))
+	UPROPERTY(EditInstanceOnly, Setter="SetEndAngle", Getter="GetEndAngle", Interp, Category="RadialArrange", meta=(ClampMin = "-180.0", UIMin = "-180.0", ClampMax = "180.0", UIMax = "180.0", AllowPrivateAccess="true"))
 	float EndAngle = 180.0f;
 
 	/** Defines how to arrange the child elements around the center. */
-	UPROPERTY(EditAnywhere, BlueprintReadWrite, Setter="SetArrangement", Getter="GetArrangement", Category="RadialArrange", meta=(AllowPrivateAccess="true"))
+	UPROPERTY(EditInstanceOnly, Setter="SetArrangement", Getter="GetArrangement", Category="RadialArrange", meta=(AllowPrivateAccess="true"))
 	EAvaRadialArrangement Arrangement = EAvaRadialArrangement::Equal;
 
 	/** If true, will arrange the child elements starting from the outer radius and moving to the inner radius. Has no effect if only using one ring. */
-	UPROPERTY(EditAnywhere, BlueprintReadWrite, Setter="SetStartFromOuterRadius", Getter="GetStartFromOuterRadius", Category="RadialArrange", meta=(AllowPrivateAccess="true"))
+	UPROPERTY(EditInstanceOnly, Setter="SetStartFromOuterRadius", Getter="GetStartFromOuterRadius", Category="RadialArrange", meta=(AllowPrivateAccess="true"))
 	bool bStartFromOuterRadius;
 
 	/** If true, will orient the selected axis torwards the center. */
-	UPROPERTY(EditAnywhere, BlueprintReadWrite, Setter="SetOrient", Getter="GetOrient", Category="RadialArrange", meta=(AllowPrivateAccess="true"))
+	UPROPERTY(EditInstanceOnly, Setter="SetOrient", Getter="GetOrient", Category="RadialArrange", meta=(AllowPrivateAccess="true"))
 	bool bOrient;
 
 	/** The axis to look at the center. */
-	UPROPERTY(EditAnywhere, BlueprintReadWrite, Setter="SetOrientAxis", Getter="GetOrientAxis", Category="RadialArrange", meta=(EditCondition="bOrient", EditConditionHides, AllowPrivateAccess="true"))
+	UPROPERTY(EditInstanceOnly, Setter="SetOrientAxis", Getter="GetOrientAxis", Category="RadialArrange", meta=(EditCondition="bOrient", EditConditionHides, AllowPrivateAccess="true"))
 	EAvaAxis OrientAxis;
 
 	/** If true, will flip the orientation axis to the opposite direction. */
-	UPROPERTY(EditAnywhere, BlueprintReadWrite, Setter="SetFlipOrient", Getter="GetFlipOrient", Category="RadialArrange", meta=(EditCondition="bOrient", EditConditionHides, AllowPrivateAccess="true"))
+	UPROPERTY(EditInstanceOnly, Setter="SetFlipOrient", Getter="GetFlipOrient", Category="RadialArrange", meta=(EditCondition="bOrient", EditConditionHides, AllowPrivateAccess="true"))
 	bool bFlipOrient;
 };
diff --git a/Engine/Plugins/Experimental/Avalanche/Source/AvalancheModifiers/Public/Modifiers/AvaSplineSweepModifier.h b/Engine/Plugins/Experimental/Avalanche/Source/AvalancheModifiers/Public/Modifiers/AvaSplineSweepModifier.h
--- a/Engine/Plugins/Experimental/Avalanche/Source/AvalancheModifiers/Public/Modifiers/AvaSplineSweepModifier.h
+++ b/Engine/Plugins/Experimental/Avalanche/Source/AvalancheModifiers/Public/Modifiers/AvaSplineSweepModifier.h
@@ -25,7 +25,19 @@ class UAvaSplineSweepModifier : public UAvaGeometryBaseModifier
 	GENERATED_BODY()
 
 public:
-	static inline constexpr int32 MaxSampleCount = 100;
+	static constexpr int32 MaxSampleCount = 100;
+
+	UFUNCTION(BlueprintCallable, Category="Motion Design|Modifiers|SplineSweep")
+	AVALANCHEMODIFIERS_API void SetSplineActor(AActor* InActor)
+	{
+		SetSplineActorWeak(InActor);
+	}
+
+	UFUNCTION(BlueprintPure, Category="Motion Design|Modifiers|SplineSweep")
+	AActor* GetSplineActor() const
+	{
+		return SplineActorWeak.Get();
+	}
 
 	AVALANCHEMODIFIERS_API void SetSplineActorWeak(TWeakObjectPtr<AActor> InSplineActorWeak);
 	TWeakObjectPtr<AActor> GetSplineActorWeak() const
@@ -33,91 +45,91 @@ public:
 		return SplineActorWeak;
 	}
 
-	UFUNCTION(BlueprintCallable, Category = "Motion Design|Modifiers|SplineSweep")
+	UFUNCTION(BlueprintCallable, Category="Motion Design|Modifiers|SplineSweep")
 	AVALANCHEMODIFIERS_API void SetSampleMode(EAvaSplineSweepSampleMode InMode);
 
-	UFUNCTION(BlueprintPure, Category = "Motion Design|Modifiers|SplineSweep")
+	UFUNCTION(BlueprintPure, Category="Motion Design|Modifiers|SplineSweep")
 	EAvaSplineSweepSampleMode GetSampleMode() const
 	{
 		return SampleMode;
 	}
 
-	UFUNCTION(BlueprintCallable, Category = "Motion Design|Modifiers|SplineSweep")
+	UFUNCTION(BlueprintCallable, Category="Motion Design|Modifiers|SplineSweep")
 	AVALANCHEMODIFIERS_API void SetSampleDistance(float InDistance);
 
-	UFUNCTION(BlueprintPure, Category = "Motion Design|Modifiers|SplineSweep")
+	UFUNCTION(BlueprintPure, Category="Motion Design|Modifiers|SplineSweep")
 	float GetSampleDistance() const
 	{
 		return SampleDistance;
 	}
 
-	UFUNCTION(BlueprintCallable, Category = "Motion Design|Modifiers|SplineSweep")
+	UFUNCTION(BlueprintCallable, Category="Motion Design|Modifiers|SplineSweep")
 	AVALANCHEMODIFIERS_API void SetSteps(int32 InSteps);
 
-	UFUNCTION(BlueprintPure, Category = "Motion Design|Modifiers|SplineSweep")
+	UFUNCTION(BlueprintPure, Category="Motion Design|Modifiers|SplineSweep")
 	int32 GetSteps() const
 	{
 		return Steps;
 	}
 
-	UFUNCTION(BlueprintCallable, Category = "Motion Design|Modifiers|SplineSweep")
+	UFUNCTION(BlueprintCallable, Category="Motion Design|Modifiers|SplineSweep")
 	AVALANCHEMODIFIERS_API void SetProgressOffset(float InOffset);
 
-	UFUNCTION(BlueprintPure, Category = "Motion Design|Modifiers|SplineSweep")
+	UFUNCTION(BlueprintPure, Category="Motion Design|Modifiers|SplineSweep")
 	float GetProgressOffset() const
 	{
 		return ProgressOffset;
 	}
 
-	UFUNCTION(BlueprintCallable, Category = "Motion Design|Modifiers|SplineSweep")
+	UFUNCTION(BlueprintCallable, Category="Motion Design|Modifiers|SplineSweep")
 	AVALANCHEMODIFIERS_API void SetProgressStart(float InStart);
 
-	UFUNCTION(BlueprintPure, Category = "Motion Design|Modifiers|SplineSweep")
+	UFUNCTION(BlueprintPure, Category="Motion Design|Modifiers|SplineSweep")
 	float GetProgressStart() const
 	{
 		return ProgressStart;
 	}
 
-	UFUNCTION(BlueprintCallable, Category = "Motion Design|Modifiers|SplineSweep")
+	UFUNCTION(BlueprintCallable, Category="Motion Design|Modifiers|SplineSweep")
 	AVALANCHEMODIFIERS_API void SetProgressEnd(float InEnd);
 
-	UFUNCTION(BlueprintPure, Category = "Motion Design|Modifiers|SplineSweep")
+	UFUNCTION(BlueprintPure, Category="Motion Design|Modifiers|SplineSweep")
 	float GetProgressEnd() const
 	{
 		return ProgressEnd;
 	}
 
-	UFUNCTION(BlueprintCallable, Category = "Motion Design|Modifiers|SplineSweep")
+	UFUNCTION(BlueprintCallable, Category="Motion Design|Modifiers|SplineSweep")
 	AVALANCHEMODIFIERS_API void SetScaleStart(float InScaleStart);
 
-	UFUNCTION(BlueprintPure, Category = "Motion Design|Modifiers|SplineSweep")
+	UFUNCTION(BlueprintPure, Category="Motion Design|Modifiers|SplineSweep")
 	float GetScaleStart() const
 	{
 		return ScaleStart;
 	}
 
-	UFUNCTION(BlueprintCallable, Category = "Motion Design|Modifiers|SplineSweep")
+	UFUNCTION(BlueprintCallable, Category="Motion Design|Modifiers|SplineSweep")
 	AVALANCHEMODIFIERS_API void SetScaleEnd(float InScaleEnd);
 
-	UFUNCTION(BlueprintPure, Category = "Motion Design|Modifiers|SplineSweep")
+	UFUNCTION(BlueprintPure, Category="Motion Design|Modifiers|SplineSweep")
 	float GetScaleEnd() const
 	{
 		return ScaleEnd;
 	}
 
-	UFUNCTION(BlueprintCallable, Category = "Motion Design|Modifiers|SplineSweep")
+	UFUNCTION(BlueprintCallable, Category="Motion Design|Modifiers|SplineSweep")
 	AVALANCHEMODIFIERS_API void SetCapped(bool bInCapped);
 
-	UFUNCTION(BlueprintPure, Category = "Motion Design|Modifiers|SplineSweep")
+	UFUNCTION(BlueprintPure, Category="Motion Design|Modifiers|SplineSweep")
 	bool GetCapped() const
 	{
 		return bCapped;
 	}
 
-	UFUNCTION(BlueprintCallable, Category = "Motion Design|Modifiers|SplineSweep")
+	UFUNCTION(BlueprintCallable, Category="Motion Design|Modifiers|SplineSweep")
 	AVALANCHEMODIFIERS_API void SetLooped(bool bInLooped);
 
-	UFUNCTION(BlueprintPure, Category = "Motion Design|Modifiers|SplineSweep")
+	UFUNCTION(BlueprintPure, Category="Motion Design|Modifiers|SplineSweep")
 	bool GetLooped() const
 	{
 		return bLooped;
@@ -147,47 +159,47 @@ protected:
 	bool SampleSpline();
 
 	/** Spline actor to retrieve the USplineComponent from */
-	UPROPERTY(EditInstanceOnly, BlueprintReadWrite, Setter="SetSplineActorWeak", Getter="GetSplineActorWeak", Category="SplineSweep", meta=(AllowPrivateAccess="true"))
+	UPROPERTY(EditInstanceOnly, Setter="SetSplineActorWeak", Getter="GetSplineActorWeak", Category="SplineSweep", meta=(AllowPrivateAccess="true"))
 	TWeakObjectPtr<AActor> SplineActorWeak;
 
 	/** How do we sample the spline */
-	UPROPERTY(EditInstanceOnly, BlueprintReadWrite, Setter="SetSampleMode", Getter="GetSampleMode", Category="SplineSweep", meta=(AllowPrivateAccess="true"))
+	UPROPERTY(EditInstanceOnly, Setter="SetSampleMode", Getter="GetSampleMode", Category="SplineSweep", meta=(AllowPrivateAccess="true"))
 	EAvaSplineSweepSampleMode SampleMode = EAvaSplineSweepSampleMode::FullDistance;
 
 	/** Custom sample distance for steps per distance mode */
-	UPROPERTY(EditInstanceOnly, BlueprintReadWrite, Setter="SetSampleDistance", Getter="GetSampleDistance", Category="SplineSweep", meta=(ClampMin="1", EditCondition="SampleMode == EAvaSplineSweepSampleMode::CustomDistance", EditConditionHides, AllowPrivateAccess="true"))
+	UPROPERTY(EditInstanceOnly, Setter="SetSampleDistance", Getter="GetSampleDistance", Category="SplineSweep", meta=(ClampMin="1", EditCondition="SampleMode == EAvaSplineSweepSampleMode::CustomDistance", EditConditionHides, AllowPrivateAccess="true"))
 	float SampleDistance = 1000.f;
 
 	/** The sample count defines the precision of the sweep */
-	UPROPERTY(EditInstanceOnly, BlueprintReadWrite, Setter="SetSteps", Getter="GetSteps", Category="SplineSweep", meta=(ClampMin="0", ClampMax="100", AllowPrivateAccess="true"))
+	UPROPERTY(EditInstanceOnly, Setter="SetSteps", Getter="GetSteps", Category="SplineSweep", meta=(ClampMin="0", ClampMax="100", AllowPrivateAccess="true"))
 	int32 Steps = 10;
 
 	/** Sample range offset of the spline, for closed loop spline range can wrap around */
-	UPROPERTY(EditInstanceOnly, BlueprintReadWrite, Setter="SetProgressOffset", Getter="GetProgressOffset", Category="SplineSweep", meta=(ClampMin="0", AllowPrivateAccess="true"))
+	UPROPERTY(EditInstanceOnly, Setter="SetProgressOffset", Getter="GetProgressOffset", Category="SplineSweep", meta=(ClampMin="0", AllowPrivateAccess="true"))
 	float ProgressOffset = 0.f;
 
 	/** Sample start range of the spline */
-	UPROPERTY(EditInstanceOnly, BlueprintReadWrite, Setter="SetProgressStart", Getter="GetProgressStart", Category="SplineSweep", meta=(ClampMin="0", ClampMax="1", AllowPrivateAccess="true"))
+	UPROPERTY(EditInstanceOnly, Setter="SetProgressStart", Getter="GetProgressStart", Category="SplineSweep", meta=(ClampMin="0", ClampMax="1", AllowPrivateAccess="true"))
 	float ProgressStart = 0.f;
 
 	/** Sample end range of the spline */
-	UPROPERTY(EditInstanceOnly, BlueprintReadWrite, Setter="SetProgressEnd", Getter="GetProgressEnd", Category="SplineSweep", meta=(ClampMin="0", ClampMax="1", AllowPrivateAccess="true"))
+	UPROPERTY(EditInstanceOnly, Setter="SetProgressEnd", Getter="GetProgressEnd", Category="SplineSweep", meta=(ClampMin="0", ClampMax="1", AllowPrivateAccess="true"))
 	float ProgressEnd = 1.f;
 
 	/** Start scale of the spline mesh */
-	UPROPERTY(EditInstanceOnly, BlueprintReadWrite, Setter="SetScaleStart", Getter="GetScaleStart", Category="SplineSweep", meta=(ClampMin="0", AllowPrivateAccess="true"))
+	UPROPERTY(EditInstanceOnly, Setter="SetScaleStart", Getter="GetScaleStart", Category="SplineSweep", meta=(ClampMin="0", AllowPrivateAccess="true"))
 	float ScaleStart = 1.f;
 
 	/** End scale of the spline mesh */
-	UPROPERTY(EditInstanceOnly, BlueprintReadWrite, Setter="SetScaleEnd", Getter="GetScaleEnd", Category="SplineSweep", meta=(ClampMin="0", AllowPrivateAccess="true"))
+	UPROPERTY(EditInstanceOnly, Setter="SetScaleEnd", Getter="GetScaleEnd", Category="SplineSweep", meta=(ClampMin="0", AllowPrivateAccess="true"))
 	float ScaleEnd = 1.f;
 
 	/** Whether start and end are closed, loop must be false */
-	UPROPERTY(EditInstanceOnly, BlueprintReadWrite, Setter="SetCapped", Getter="GetCapped", Category="SplineSweep", meta=(EditCondition="!bLooped", AllowPrivateAccess="true"))
+	UPROPERTY(EditInstanceOnly, Setter="SetCapped", Getter="GetCapped", Category="SplineSweep", meta=(EditCondition="!bLooped", AllowPrivateAccess="true"))
 	bool bCapped = true;
 
 	/** Whether we close the whole spline path, if spline loops this will be true */
-	UPROPERTY(EditInstanceOnly, BlueprintReadWrite, Setter="SetLooped", Getter="GetLooped", Category="SplineSweep", meta=(AllowPrivateAccess="true"))
+	UPROPERTY(EditInstanceOnly, Setter="SetLooped", Getter="GetLooped", Category="SplineSweep", meta=(AllowPrivateAccess="true"))
 	bool bLooped = false;
 
 	/** Spline component to sample from */
diff --git a/Engine/Plugins/Experimental/Avalanche/Source/AvalancheModifiers/Public/Modifiers/AvaSubdivideModifier.h b/Engine/Plugins/Experimental/Avalanche/Source/AvalancheModifiers/Public/Modifiers/AvaSubdivideModifier.h
--- a/Engine/Plugins/Experimental/Avalanche/Source/AvalancheModifiers/Public/Modifiers/AvaSubdivideModifier.h
+++ b/Engine/Plugins/Experimental/Avalanche/Source/AvalancheModifiers/Public/Modifiers/AvaSubdivideModifier.h
@@ -5,7 +5,7 @@
 #include "AvaGeometryBaseModifier.h"
 #include "AvaSubdivideModifier.generated.h"
 
-UENUM()
+UENUM(BlueprintType)
 enum class EAvaSubdivisionType : uint8
 {
 	Selective,
@@ -22,19 +22,28 @@ public:
 	static constexpr int MinSubdivideCuts = 1;
     static constexpr int MaxSubdivideCuts = 15;
 
+	UFUNCTION(BlueprintCallable, Category="Motion Design|Modifiers|Subdivide")
 	AVALANCHEMODIFIERS_API void SetCuts(int32 InCuts);
+
+	UFUNCTION(BlueprintPure, Category="Motion Design|Modifiers|Subdivide")
 	int32 GetCuts() const
 	{
 		return Cuts;
 	}
 
+	UFUNCTION(BlueprintCallable, Category="Motion Design|Modifiers|Subdivide")
 	AVALANCHEMODIFIERS_API void SetRecomputeNormals(bool bInRecomputeNormals);
+
+	UFUNCTION(BlueprintPure, Category="Motion Design|Modifiers|Subdivide")
 	bool GetRecomputeNormals() const
 	{
 		return bRecomputeNormals;
 	}
 
+	UFUNCTION(BlueprintCallable, Category="Motion Design|Modifiers|Subdivide")
 	AVALANCHEMODIFIERS_API void SetType(EAvaSubdivisionType InType);
+
+	UFUNCTION(BlueprintPure, Category="Motion Design|Modifiers|Subdivide")
 	EAvaSubdivisionType GetType() const
 	{
 		return Type;
@@ -54,12 +63,12 @@ protected:
 
 	void OnOptionsChanged();
 
-	UPROPERTY(EditInstanceOnly, BlueprintReadWrite, Setter="SetCuts", Getter="GetCuts", Category="Subdivide", meta=(ClampMin="1", ClampMax="15", AllowPrivateAccess="true"))
+	UPROPERTY(EditInstanceOnly, Setter="SetCuts", Getter="GetCuts", Category="Subdivide", meta=(ClampMin="1", ClampMax="15", AllowPrivateAccess="true"))
 	int32 Cuts = 2;
 
-	UPROPERTY(EditInstanceOnly, BlueprintReadWrite, Setter="SetRecomputeNormals", Getter="GetRecomputeNormals", Category="Subdivide", meta=(EditCondition="Type == EAvaSubdivisionType::PN", EditConditionHides, AllowPrivateAccess="true"))
+	UPROPERTY(EditInstanceOnly, Setter="SetRecomputeNormals", Getter="GetRecomputeNormals", Category="Subdivide", meta=(EditCondition="Type == EAvaSubdivisionType::PN", EditConditionHides, AllowPrivateAccess="true"))
 	bool bRecomputeNormals = true;
 
-	UPROPERTY(EditInstanceOnly, BlueprintReadWrite, Setter="SetType", Getter="GetType", Category="Subdivide", meta=(AllowPrivateAccess="true"))
+	UPROPERTY(EditInstanceOnly, Setter="SetType", Getter="GetType", Category="Subdivide", meta=(AllowPrivateAccess="true"))
 	EAvaSubdivisionType Type = EAvaSubdivisionType::Uniform;
 };
diff --git a/Engine/Plugins/Experimental/Avalanche/Source/AvalancheModifiers/Public/Modifiers/AvaTaperModifier.h b/Engine/Plugins/Experimental/Avalanche/Source/AvalancheModifiers/Public/Modifiers/AvaTaperModifier.h
--- a/Engine/Plugins/Experimental/Avalanche/Source/AvalancheModifiers/Public/Modifiers/AvaTaperModifier.h
+++ b/Engine/Plugins/Experimental/Avalanche/Source/AvalancheModifiers/Public/Modifiers/AvaTaperModifier.h
@@ -6,14 +6,14 @@
 #include "Tools/AvaTaperTool.h"
 #include "AvaTaperModifier.generated.h"
 
-UENUM()
+UENUM(BlueprintType)
 enum class EAvaTaperReferenceFrame : uint8
 {
 	MeshCenter,
 	Custom
 };
 
-UENUM()
+UENUM(BlueprintType)
 enum class EAvaTaperExtent : uint8
 {
 	WholeShape,
@@ -29,49 +29,73 @@ public:
 	static constexpr int MinTaperLatticeResolution = 1;
 	static constexpr int MaxTaperLatticeResolution = 20;
 
+	UFUNCTION(BlueprintCallable, Category="Motion Design|Modifiers|Taper")
 	AVALANCHEMODIFIERS_API void SetAmount(float InAmount);
+
+	UFUNCTION(BlueprintPure, Category="Motion Design|Modifiers|Taper")
 	float GetAmount() const
 	{
 		return Amount;
 	}
 
+	UFUNCTION(BlueprintCallable, Category="Motion Design|Modifiers|Taper")
 	AVALANCHEMODIFIERS_API void SetUpperExtent(float InUpperExtent);
+
+	UFUNCTION(BlueprintPure, Category="Motion Design|Modifiers|Taper")
 	float GetUpperExtent() const
 	{
 		return UpperExtent;
 	}
 
+	UFUNCTION(BlueprintCallable, Category="Motion Design|Modifiers|Taper")
 	AVALANCHEMODIFIERS_API void SetLowerExtent(float InLowerExtent);
+
+	UFUNCTION(BlueprintPure, Category="Motion Design|Modifiers|Taper")
 	float GetLowerExtent() const
 	{
 		return LowerExtent;
 	}
 
+	UFUNCTION(BlueprintCallable, Category="Motion Design|Modifiers|Taper")
 	AVALANCHEMODIFIERS_API void SetExtent(EAvaTaperExtent InExtent);
+
+	UFUNCTION(BlueprintPure, Category="Motion Design|Modifiers|Taper")
 	EAvaTaperExtent GetExtent() const
 	{
 		return Extent;
 	}
 
+	UFUNCTION(BlueprintCallable, Category="Motion Design|Modifiers|Taper")
 	AVALANCHEMODIFIERS_API void SetInterpolationType(EAvaTaperInterpolationType InInterpolationType);
+
+	UFUNCTION(BlueprintPure, Category="Motion Design|Modifiers|Taper")
 	EAvaTaperInterpolationType GetInterpolationType() const
 	{
 		return InterpolationType;
 	}
 
+	UFUNCTION(BlueprintCallable, Category="Motion Design|Modifiers|Taper")
 	AVALANCHEMODIFIERS_API void SetReferenceFrame(EAvaTaperReferenceFrame InReferenceFrame);
+
+	UFUNCTION(BlueprintPure, Category="Motion Design|Modifiers|Taper")
 	EAvaTaperReferenceFrame GetReferenceFrame() const
 	{
 		return ReferenceFrame;
 	}
 
+	UFUNCTION(BlueprintCallable, Category="Motion Design|Modifiers|Taper")
 	AVALANCHEMODIFIERS_API void SetResolution(int32 InResolution);
+
+	UFUNCTION(BlueprintPure, Category="Motion Design|Modifiers|Taper")
 	int32 GetResolution() const
 	{
 		return Resolution;
 	}
 
+	UFUNCTION(BlueprintCallable, Category="Motion Design|Modifiers|Taper")
 	AVALANCHEMODIFIERS_API void SetOffset(FVector2D InOffset);
+
+	UFUNCTION(BlueprintPure, Category="Motion Design|Modifiers|Taper")
 	FVector2D GetOffset() const
 	{
 		return Offset;
@@ -98,28 +122,28 @@ protected:
 
 	int32 GetSubdividersCuts() const;
 
-	UPROPERTY(EditInstanceOnly, BlueprintReadWrite, Setter="SetAmount", Getter="GetAmount", Category="Taper", meta=(ClampMin="0.0", ClampMax="1.0", AllowPrivateAccess="true"))
+	UPROPERTY(EditInstanceOnly, Setter="SetAmount", Getter="GetAmount", Category="Taper", meta=(ClampMin="0.0", ClampMax="1.0", AllowPrivateAccess="true"))
 	float Amount = 0.0;
 
-	UPROPERTY(EditInstanceOnly, BlueprintReadWrite, Setter="SetExtent", Getter="GetExtent", Category="Taper", meta=(AllowPrivateAccess="true"))
+	UPROPERTY(EditInstanceOnly, Setter="SetExtent", Getter="GetExtent", Category="Taper", meta=(AllowPrivateAccess="true"))
 	EAvaTaperExtent Extent = EAvaTaperExtent::WholeShape;
 
-	UPROPERTY(EditInstanceOnly, BlueprintReadWrite, Setter="SetUpperExtent", Getter="GetUpperExtent", Category="Taper", meta=(EditCondition="Extent == EAvaTaperExtent::Custom", EditConditionHides, ClampMin="0", ClampMax="100", Units="Percent", ToolTip="100%: shape top.\n0%: shape bottom.", AllowPrivateAccess="true"))
+	UPROPERTY(EditInstanceOnly, Setter="SetUpperExtent", Getter="GetUpperExtent", Category="Taper", meta=(EditCondition="Extent == EAvaTaperExtent::Custom", EditConditionHides, ClampMin="0", ClampMax="100", Units="Percent", ToolTip="100%: shape top.\n0%: shape bottom.", AllowPrivateAccess="true"))
 	float UpperExtent = 100;
 
-	UPROPERTY(EditInstanceOnly, BlueprintReadWrite, Setter="SetLowerExtent", Getter="GetLowerExtent", Category="Taper", meta=(EditCondition="Extent == EAvaTaperExtent::Custom", EditConditionHides, ClampMin="0", ClampMax="100", Units="Percent", ToolTip="100%: shape bottom.\n0%: shape top.", AllowPrivateAccess="true"))
+	UPROPERTY(EditInstanceOnly, Setter="SetLowerExtent", Getter="GetLowerExtent", Category="Taper", meta=(EditCondition="Extent == EAvaTaperExtent::Custom", EditConditionHides, ClampMin="0", ClampMax="100", Units="Percent", ToolTip="100%: shape bottom.\n0%: shape top.", AllowPrivateAccess="true"))
 	float LowerExtent = 100;
 
-	UPROPERTY(EditInstanceOnly, BlueprintReadWrite, Setter="SetInterpolationType", Getter="GetInterpolationType", Category="Taper", meta=(AllowPrivateAccess="true"))
+	UPROPERTY(EditInstanceOnly, Setter="SetInterpolationType", Getter="GetInterpolationType", Category="Taper", meta=(AllowPrivateAccess="true"))
 	EAvaTaperInterpolationType InterpolationType = EAvaTaperInterpolationType::Linear;
 
-	UPROPERTY(EditInstanceOnly, BlueprintReadWrite, Setter="SetResolution", Getter="GetResolution", Category="Taper", meta=(ClampMin="1", ClampMax="20", ToolTip="The number of vertical control points used to apply the taper. If the modifier is in a stack with Subdivide modifiers, taper will use the max value between Resolution and the total subdivision cuts.", AllowPrivateAccess="true"))
+	UPROPERTY(EditInstanceOnly, Setter="SetResolution", Getter="GetResolution", Category="Taper", meta=(ClampMin="1", ClampMax="20", ToolTip="The number of vertical control points used to apply the taper. If the modifier is in a stack with Subdivide modifiers, taper will use the max value between Resolution and the total subdivision cuts.", AllowPrivateAccess="true"))
 	int32 Resolution = 5;
 
-	UPROPERTY(EditInstanceOnly, BlueprintReadWrite, Setter="SetReferenceFrame", Getter="GetReferenceFrame", Category="Taper", meta=(AllowPrivateAccess="true"))
+	UPROPERTY(EditInstanceOnly, Setter="SetReferenceFrame", Getter="GetReferenceFrame", Category="Taper", meta=(AllowPrivateAccess="true"))
 	EAvaTaperReferenceFrame ReferenceFrame = EAvaTaperReferenceFrame::MeshCenter;
 
-	UPROPERTY(EditInstanceOnly, BlueprintReadWrite, Setter="SetOffset", Getter="GetOffset", Category="Taper", meta=(EditCondition="ReferenceFrame == EAvaTaperReferenceFrame::Custom", EditConditionHides, AllowPrivateAccess="true"))
+	UPROPERTY(EditInstanceOnly, Setter="SetOffset", Getter="GetOffset", Category="Taper", meta=(EditCondition="ReferenceFrame == EAvaTaperReferenceFrame::Custom", EditConditionHides, AllowPrivateAccess="true"))
 	FVector2D Offset = FVector2D::ZeroVector;
 
 	UPROPERTY()
diff --git a/Engine/Plugins/Experimental/Avalanche/Source/AvalancheModifiers/Public/Modifiers/AvaTranslucentPriorityModifier.h b/Engine/Plugins/Experimental/Avalanche/Source/AvalancheModifiers/Public/Modifiers/AvaTranslucentPriorityModifier.h
--- a/Engine/Plugins/Experimental/Avalanche/Source/AvalancheModifiers/Public/Modifiers/AvaTranslucentPriorityModifier.h
+++ b/Engine/Plugins/Experimental/Avalanche/Source/AvalancheModifiers/Public/Modifiers/AvaTranslucentPriorityModifier.h
@@ -5,6 +5,7 @@
 #include "AvaArrangeBaseModifier.h"
 #include "AvaBaseModifier.h"
 #include "Shared/AvaTranslucentPriorityModifierShared.h"
+#include "UObject/WeakObjectPtrTemplatesFwd.h"
 #include "AvaTranslucentPriorityModifier.generated.h"
 
 class ACameraActor;
@@ -29,37 +30,62 @@ class UAvaTranslucentPriorityModifier : public UAvaArrangeBaseModifier
 	friend class UAvaTranslucentPriorityModifierShared;
 
 public:
+	UFUNCTION(BlueprintCallable, Category="Motion Design|Modifiers|TranslucentPriority")
     AVALANCHEMODIFIERS_API void SetMode(EAvaTranslucentPriorityModifierMode InMode);
+
+	UFUNCTION(BlueprintPure, Category="Motion Design|Modifiers|TranslucentPriority")
     EAvaTranslucentPriorityModifierMode GetMode() const
     {
 	    return Mode;
     }
 
+	UFUNCTION(BlueprintCallable, Category="Motion Design|Modifiers|TranslucentPriority")
+	AVALANCHEMODIFIERS_API void SetCameraActor(ACameraActor* InCameraActor);
+
+	UFUNCTION(BlueprintPure, Category="Motion Design|Modifiers|TranslucentPriority")
+	ACameraActor* GetCameraActor() const
+	{
+		return CameraActorWeak.Get();
+	}
+
     AVALANCHEMODIFIERS_API void SetCameraActorWeak(const TWeakObjectPtr<ACameraActor>& InCameraActor);
+
 	TWeakObjectPtr<ACameraActor> GetCameraActorWeak() const
 	{
 		return CameraActorWeak;
 	}
 
+	UFUNCTION(BlueprintCallable, Category="Motion Design|Modifiers|TranslucentPriority")
     AVALANCHEMODIFIERS_API void SetSortPriority(int32 InSortPriority);
+
+	UFUNCTION(BlueprintPure, Category="Motion Design|Modifiers|TranslucentPriority")
     int32 GetSortPriority() const
     {
 	    return SortPriority;
     }
 
+	UFUNCTION(BlueprintCallable, Category="Motion Design|Modifiers|TranslucentPriority")
 	AVALANCHEMODIFIERS_API void SetSortPriorityOffset(int32 InOffset);
+
+	UFUNCTION(BlueprintPure, Category="Motion Design|Modifiers|TranslucentPriority")
 	int32 GetSortPriorityOffset() const
 	{
 		return SortPriorityOffset;
 	}
 
+	UFUNCTION(BlueprintCallable, Category="Motion Design|Modifiers|TranslucentPriority")
 	AVALANCHEMODIFIERS_API void SetSortPriorityStep(int32 InStep);
+
+	UFUNCTION(BlueprintPure, Category="Motion Design|Modifiers|TranslucentPriority")
 	int32 GetSortPriorityStep() const
 	{
 		return SortPriorityStep;
 	}
 
+	UFUNCTION(BlueprintCallable, Category="Motion Design|Modifiers|TranslucentPriority")
 	AVALANCHEMODIFIERS_API void SetIncludeChildren(bool bInIncludeChildren);
+
+	UFUNCTION(BlueprintPure, Category="Motion Design|Modifiers|TranslucentPriority")
 	bool GetIncludeChildren() const
 	{
 		return bIncludeChildren;
@@ -109,27 +135,27 @@ protected:
 	ACameraActor* GetDefaultCameraActor() const;
 
 	/** The sort mode we are currently in */
-	UPROPERTY(EditInstanceOnly, BlueprintReadWrite, Setter, Getter, Category="TranslucentPriority", meta=(AllowPrivateAccess="true"))
+	UPROPERTY(EditInstanceOnly, Setter, Getter, Category="TranslucentPriority", meta=(AllowPrivateAccess="true"))
 	EAvaTranslucentPriorityModifierMode Mode = EAvaTranslucentPriorityModifierMode::Manual;
 
 	/** The camera actor to compute the distance from */
-	UPROPERTY(EditInstanceOnly, BlueprintReadWrite, Setter, Getter, Category="TranslucentPriority", meta=(DisplayName="CameraActor", EditCondition="Mode == EAvaTranslucentPriorityModifierMode::AutoCameraDistance", AllowPrivateAccess="true"))
+	UPROPERTY(EditInstanceOnly, Setter, Getter, Category="TranslucentPriority", meta=(DisplayName="CameraActor", EditCondition="Mode == EAvaTranslucentPriorityModifierMode::AutoCameraDistance", AllowPrivateAccess="true"))
 	TWeakObjectPtr<ACameraActor> CameraActorWeak = nullptr;
 
 	/** The sort priority that will be set on the primitive component for manual mode */
-	UPROPERTY(EditInstanceOnly, BlueprintReadWrite, Setter, Getter, Category="TranslucentPriority", meta=(EditCondition="Mode == EAvaTranslucentPriorityModifierMode::Manual", EditConditionHides, AllowPrivateAccess="true"))
+	UPROPERTY(EditInstanceOnly, Setter, Getter, Category="TranslucentPriority", meta=(EditCondition="Mode == EAvaTranslucentPriorityModifierMode::Manual", EditConditionHides, AllowPrivateAccess="true"))
 	int32 SortPriority = 0;
 
 	/** Sort priority offset shared across all modifiers in this same level */
-	UPROPERTY(Transient, EditInstanceOnly, BlueprintReadWrite, Setter, Getter, Category="TranslucentPriority", meta=(AllowPrivateAccess="true"))
+	UPROPERTY(Transient, EditInstanceOnly, Setter, Getter, Category="TranslucentPriority", meta=(AllowPrivateAccess="true"))
 	int32 SortPriorityOffset = 0;
 
 	/** Sort priority incremental step shared across all modifiers in this same level */
-	UPROPERTY(Transient, EditInstanceOnly, BlueprintReadWrite, Setter, Getter, Category="TranslucentPriority", meta=(AllowPrivateAccess="true"))
+	UPROPERTY(Transient, EditInstanceOnly, Setter, Getter, Category="TranslucentPriority", meta=(AllowPrivateAccess="true"))
 	int32 SortPriorityStep = 1;
 
 	/** If true, will include children too and update their sort priority */
-	UPROPERTY(EditAnywhere, BlueprintReadWrite, Setter="SetIncludeChildren", Getter="GetIncludeChildren", Category="TranslucentPriority", meta=(AllowPrivateAccess="true"))
+	UPROPERTY(EditInstanceOnly, Setter="SetIncludeChildren", Getter="GetIncludeChildren", Category="TranslucentPriority", meta=(AllowPrivateAccess="true"))
 	bool bIncludeChildren = true;
 
 	/** The components this modifier is managing */
diff --git a/Engine/Plugins/Experimental/Avalanche/Source/AvalancheModifiers/Public/Modifiers/AvaVisibilityModifier.h b/Engine/Plugins/Experimental/Avalanche/Source/AvalancheModifiers/Public/Modifiers/AvaVisibilityModifier.h
--- a/Engine/Plugins/Experimental/Avalanche/Source/AvalancheModifiers/Public/Modifiers/AvaVisibilityModifier.h
+++ b/Engine/Plugins/Experimental/Avalanche/Source/AvalancheModifiers/Public/Modifiers/AvaVisibilityModifier.h
@@ -17,29 +17,29 @@ class UAvaVisibilityModifier : public UAvaArrangeBaseModifier
 
 public:
 	/** Sets the child index range to hide instead of showing. */
-	UFUNCTION(BlueprintCallable, Category = "Motion Design|Modifiers|Visibility")
-	AVALANCHEMODIFIERS_API void SetInvertVisibility(const bool bNewInvertVisibility);
+	UFUNCTION(BlueprintCallable, Category="Motion Design|Modifiers|Visibility")
+	AVALANCHEMODIFIERS_API void SetInvertVisibility(const bool bInInvertVisibility);
 
 	/** Returns true if hiding the child index range instead of showing. */
-	UFUNCTION(BlueprintPure, Category = "Motion Design|Modifiers|Visibility")
+	UFUNCTION(BlueprintPure, Category="Motion Design|Modifiers|Visibility")
 	bool GetInvertVisibility() const
 	{
 		return bInvertVisibility;
 	}
 
-	UFUNCTION(BlueprintCallable, Category = "Motion Design|Modifiers|Visibility")
+	UFUNCTION(BlueprintCallable, Category="Motion Design|Modifiers|Visibility")
 	AVALANCHEMODIFIERS_API void SetIndex(int32 InIndex);
 
-	UFUNCTION(BlueprintPure, Category = "Motion Design|Modifiers|Visibility")
+	UFUNCTION(BlueprintPure, Category="Motion Design|Modifiers|Visibility")
 	int32 GetIndex() const
 	{
 		return Index;
 	}
 
-	UFUNCTION(BlueprintCallable, Category = "Motion Design|Modifiers|Visibility")
+	UFUNCTION(BlueprintCallable, Category="Motion Design|Modifiers|Visibility")
 	AVALANCHEMODIFIERS_API void SetTreatAsRange(const bool bInTreatAsRange);
 
-	UFUNCTION(BlueprintPure, Category = "Motion Design|Modifiers|Visibility")
+	UFUNCTION(BlueprintPure, Category="Motion Design|Modifiers|Visibility")
 	bool GetTreatAsRange() const
 	{
 		return bTreatAsRange;
@@ -72,15 +72,15 @@ protected:
 	AActor* GetDirectChildren(AActor* InParentActor, AActor* InChildActor);
 
 	/** Child index to set visibility, visible if bInvertVisibility is false else hidden */
-	UPROPERTY(EditAnywhere, BlueprintReadWrite, Setter="SetIndex", Getter="GetIndex", Category="Visibility", meta=(ClampMin="0", UIMin="0", AllowPrivateAccess="true"))
+	UPROPERTY(EditInstanceOnly, Setter="SetIndex", Getter="GetIndex", Category="Visibility", meta=(ClampMin="0", UIMin="0", AllowPrivateAccess="true"))
 	int32 Index = 0;
 
 	/** Treat index as a range from 0 to index */
-	UPROPERTY(EditAnywhere, BlueprintReadWrite, Setter="SetTreatAsRange", Getter="GetTreatAsRange", Category="Visibility", meta=(AllowPrivateAccess="true"))
+	UPROPERTY(EditInstanceOnly, Setter="SetTreatAsRange", Getter="GetTreatAsRange", Category="Visibility", meta=(AllowPrivateAccess="true"))
 	bool bTreatAsRange = false;
 
 	/** If true, will hide the child index range instead of showing. */
-	UPROPERTY(EditAnywhere, BlueprintReadWrite, Setter="SetInvertVisibility", Getter="GetInvertVisibility", Category="Visibility", meta=(AllowPrivateAccess="true"))
+	UPROPERTY(EditInstanceOnly, Setter="SetInvertVisibility", Getter="GetInvertVisibility", Category="Visibility", meta=(AllowPrivateAccess="true"))
 	bool bInvertVisibility = false;
 
 	/** Visibility defined for direct children of this actor, used by other modifiers of this type to check */
