MotionDesign : OperatorStack
- Adding support for item multi view to replace split view
- Allowing to group item of the same type as a single item
- General UI Improvements to match unreal editor style

diff --git a/Engine/Plugins/Experimental/ActorModifierCore/Source/ActorModifierCoreEditor/Private/Modifiers/Customizations/ActorModifierCoreEditorStackCustomization.cpp b/Engine/Plugins/Experimental/ActorModifierCore/Source/ActorModifierCoreEditor/Private/Modifiers/Customizations/ActorModifierCoreEditorStackCustomization.cpp
--- a/Engine/Plugins/Experimental/ActorModifierCore/Source/ActorModifierCoreEditor/Private/Modifiers/Customizations/ActorModifierCoreEditorStackCustomization.cpp
+++ b/Engine/Plugins/Experimental/ActorModifierCore/Source/ActorModifierCoreEditor/Private/Modifiers/Customizations/ActorModifierCoreEditorStackCustomization.cpp
@@ -18,13 +18,14 @@
 #include "Subsystems/ActorModifierCoreSubsystem.h"
 #include "ActorModifierCoreEditorStyle.h"
 #include "Contexts/OperatorStackEditorMenuContext.h"
+#include "Items/OperatorStackEditorGroupItem.h"
 #include "Items/OperatorStackEditorObjectItem.h"
 #include "Widgets/Notifications/SNotificationList.h"
 #include "Widgets/Views/STableRow.h"
 
 #define LOCTEXT_NAMESPACE "ActorModifierCoreEditorStackCustomization"
 
-DEFINE_LOG_CATEGORY_STATIC(LogActorModifierCoreEditorPropertiesWrapper, Log, All);
+DEFINE_LOG_CATEGORY_STATIC(LogActorModifierCoreEditorStackCustomization, Log, All);
 
 UActorModifierCoreEditorStackCustomization::UActorModifierCoreEditorStackCustomization()
 	: UOperatorStackEditorStackCustomization(
@@ -37,43 +38,101 @@ UActorModifierCoreEditorStackCustomization::UActorModifierCoreEditorStackCustomi
 	RegisterCustomizationFor(UActorModifierCoreBase::StaticClass());
 }
 
-bool UActorModifierCoreEditorStackCustomization::TransformContextItem(const FOperatorStackEditorItemPtr& InItem, TArray<FOperatorStackEditorItemPtr>& OutTransformedItems) const
+bool UActorModifierCoreEditorStackCustomization::GetRootItem(const FOperatorStackEditorContext& InContext, FOperatorStackEditorItemPtr& OutRootItem) const
 {
-	// If we have an actor then display stack
-	if (InItem->IsA<AActor>())
+	TArray<FOperatorStackEditorItemPtr> RootItems;
+
+	// Gather all modifiers stack as root items
+	for (const FOperatorStackEditorItemPtr& Item : InContext.GetItems())
 	{
-		const UActorModifierCoreSubsystem* ModifierSubsystem = UActorModifierCoreSubsystem::Get();
+		if (!Item.IsValid())
+		{
+			continue;
+		}
 
-		if (UActorModifierCoreStack* Stack = ModifierSubsystem->GetActorModifierStack(InItem->Get<AActor>()))
+		if (Item->IsA<AActor>())
 		{
-			OutTransformedItems.Add(MakeShared<FOperatorStackEditorObjectItem>(Stack));
+			const UActorModifierCoreSubsystem* ModifierSubsystem = UActorModifierCoreSubsystem::Get();
 
-			return true;
+			for (const AActor* Actor : Item->GetAsArray<AActor>())
+			{
+				if (UActorModifierCoreStack* Stack = ModifierSubsystem->GetActorModifierStack(Actor))
+				{
+					RootItems.Add(MakeShared<FOperatorStackEditorObjectItem>(Stack));
+				}
+			}
+		}
+		else if (Item->IsA<UActorModifierCoreComponent>())
+		{
+			for (const UActorModifierCoreComponent* Component : Item->GetAsArray<UActorModifierCoreComponent>())
+			{
+				RootItems.Add(MakeShared<FOperatorStackEditorObjectItem>(Component->GetModifierStack()));
+			}
+		}
+		else if (Item->IsA<UActorModifierCoreStack>())
+		{
+			for (UActorModifierCoreStack* ModifierStack : Item->GetAsArray<UActorModifierCoreStack>())
+			{
+				RootItems.Add(MakeShared<FOperatorStackEditorObjectItem>(ModifierStack));
+			}
+		}
+		else if (Item->IsA<UActorModifierCoreBase>())
+		{
+			for (const UActorModifierCoreBase* Modifier : Item->GetAsArray<UActorModifierCoreBase>())
+			{
+				RootItems.Add(MakeShared<FOperatorStackEditorObjectItem>(Modifier->GetRootModifierStack()));
+			}
 		}
 	}
-	// If we have a modifier component then display stack
-	else if (InItem->IsA<UActorModifierCoreComponent>())
-	{
-		const UActorModifierCoreComponent* ModifierComponent = InItem->Get<UActorModifierCoreComponent>();
 
-		OutTransformedItems.Add(MakeShared<FOperatorStackEditorObjectItem>(ModifierComponent->GetModifierStack()));
+	OutRootItem = MakeShared<FOperatorStackEditorGroupItem>(RootItems, FOperatorStackEditorItemType(UActorModifierCoreStack::StaticClass(), EOperatorStackEditorItemType::Object));
 
-		return true;
-	}
-	// If we have a stack then display modifiers inside
-	else if (InItem->IsA<UActorModifierCoreStack>())
+	return Super::GetRootItem(InContext, OutRootItem);
+}
+
+bool UActorModifierCoreEditorStackCustomization::GetChildrenItem(const FOperatorStackEditorItemPtr& InItem, TArray<FOperatorStackEditorItemPtr>& OutChildrenItems) const
+{
+	if (InItem->IsA<UActorModifierCoreStack>())
 	{
-		const UActorModifierCoreStack* ModifierStack = InItem->Get<UActorModifierCoreStack>();
+		if (InItem->GetValueCount() > 1)
+		{
+			TMap<UClass*, int32> ClassToIndex;
+			TArray<TArray<FOperatorStackEditorItemPtr>> Modifiers;
 
-		for (UActorModifierCoreBase* Modifier : ModifierStack->GetModifiers())
+			for (const UActorModifierCoreStack* ModifierStack : InItem->GetAsArray<UActorModifierCoreStack>())
+			{
+				for (UActorModifierCoreBase* Modifier : ModifierStack->GetModifiers())
+				{
+					if (const int32* Index = ClassToIndex.Find(Modifier->GetClass()))
+					{
+						Modifiers[*Index].Add(MakeShared<FOperatorStackEditorObjectItem>(Modifier));
+					}
+					else
+					{
+						TArray<FOperatorStackEditorItemPtr> ModifierGroup;
+						ModifierGroup.Add(MakeShared<FOperatorStackEditorObjectItem>(Modifier));
+						int32 GroupIndex = Modifiers.Add(ModifierGroup);
+						ClassToIndex.Add(Modifier->GetClass(), GroupIndex);
+					}
+				}
+			}
+
+			for (int32 Index = 0; Index < Modifiers.Num(); Index++)
+			{
+				const UClass* const* ModifierClass = ClassToIndex.FindKey(Index);
+				OutChildrenItems.Add(MakeShared<FOperatorStackEditorGroupItem>(Modifiers[Index], FOperatorStackEditorItemType(*ModifierClass, EOperatorStackEditorItemType::Object)));
+			}
+		}
+		else
 		{
-			OutTransformedItems.Add(MakeShared<FOperatorStackEditorObjectItem>(Modifier));
+			for (UActorModifierCoreBase* Modifier : InItem->Get<UActorModifierCoreStack>(0)->GetModifiers())
+			{
+				OutChildrenItems.Add(MakeShared<FOperatorStackEditorObjectItem>(Modifier));
+			}
 		}
-
-		return true;
 	}
 
-	return Super::TransformContextItem(InItem, OutTransformedItems);
+	return Super::GetChildrenItem(InItem, OutChildrenItems);
 }
 
 void UActorModifierCoreEditorStackCustomization::CustomizeStackHeader(const FOperatorStackEditorTree& InItemTree, FOperatorStackEditorHeaderBuilder& InHeaderBuilder)
@@ -93,7 +152,7 @@ void UActorModifierCoreEditorStackCustomization::CustomizeStackHeader(const FOpe
 		{
 			if (Item.IsValid() && Item->IsA<UActorModifierCoreBase>())
 			{
-				const UActorModifierCoreBase* Modifier = Item->Get<UActorModifierCoreBase>();
+				const UActorModifierCoreBase* Modifier = Item->Get<UActorModifierCoreBase>(0);
 
 				if (Modifier && !Modifier->IsModifierStack())
 				{
@@ -106,7 +165,7 @@ void UActorModifierCoreEditorStackCustomization::CustomizeStackHeader(const FOpe
 			.SetToolMenu(
 				AddModifierMenuName
 				, LOCTEXT("AddModifiersMenu", "Add Modifiers")
-				, FAppStyle::GetBrush(TEXT("Icons.PlusCircle"))
+				, FAppStyle::GetBrush("Icons.Plus")
 			)
 			.SetSearchAllowed(true)
 			.SetSearchPinnedKeywords(PinnedKeywords);
@@ -122,14 +181,10 @@ void UActorModifierCoreEditorStackCustomization::CustomizeItemHeader(const FOper
 	// Customize stack and modifier header
 	if (InItem->IsA<UActorModifierCoreBase>())
 	{
-		UActorModifierCoreBase* Modifier = InItem->Get<UActorModifierCoreBase>();
-
 		FBoolProperty* ModifierEnableProperty = FindFProperty<FBoolProperty>(UActorModifierCoreBase::StaticClass(), GET_MEMBER_NAME_CHECKED(UActorModifierCoreBase, bModifierEnabled));
 
-		const bool bIsStack = Modifier->IsModifierStack();
-
 		// Commands for item on key events
-		const TSharedPtr<FUICommandList> Commands = CreateModifierCommands(Modifier);
+		const TSharedPtr<FUICommandList> Commands = CreateModifierCommands(InItem);
 
 		// Action menu available in header in slim toolbar
 		static const FName HeaderModifierMenuName = TEXT("HeaderModifierMenu");
@@ -148,6 +203,8 @@ void UActorModifierCoreEditorStackCustomization::CustomizeItemHeader(const FOper
 		}
 
 		// Item keyword for search
+		UActorModifierCoreBase* Modifier = InItem->Get<UActorModifierCoreBase>(0);
+
 		const TSet<FString> SearchKeywords
 		{
 			Modifier->GetModifierName().ToString(),
@@ -164,22 +221,24 @@ void UActorModifierCoreEditorStackCustomization::CustomizeItemHeader(const FOper
 			return true;
 		});
 
-		FString HeaderLabel = Modifier->GetModifierName().ToString();
-
-		// Add actor name next to stack label when we customize multiple items
-		if (bIsStack && InItemTree.GetRootItems().Num() > 1)
-		{
-			if (const AActor* ModifiedActor = Modifier->GetModifiedActor())
-			{
-				HeaderLabel += + TEXT(" (") + ModifiedActor->GetActorNameOrLabel() + TEXT(")");
-			}
-		}
+		const bool bIsStack = InItem->IsA<UActorModifierCoreStack>();
 
 		/** Show last execution error messages if failed execution */
 		TAttribute<EOperatorStackEditorMessageType> MessageType = EOperatorStackEditorMessageType::None;
 		TAttribute<FText> MessageText = FText::GetEmpty();
+		FString HeaderLabel = bIsStack ? TEXT("Modifiers") : Modifier->GetModifierName().ToString();
 
-		if (!bIsStack)
+		if (InItem->GetValueCount() > 1)
+		{
+			HeaderLabel += TEXT(" (") + FString::FromInt(InItem->GetValueCount()) + TEXT(")");
+
+			if (bIsStack)
+			{
+				MessageType = EOperatorStackEditorMessageType::Info;
+				MessageText = LOCTEXT("MultiModifierView", "You are viewing multiple items");
+			}
+		}
+		else if (!bIsStack)
 		{
 			TWeakObjectPtr<UActorModifierCoreBase> ModifierWeak(Modifier);
 
@@ -261,9 +320,9 @@ void UActorModifierCoreEditorStackCustomization::CustomizeItemBody(const FOperat
 void UActorModifierCoreEditorStackCustomization::CustomizeItemFooter(const FOperatorStackEditorItemPtr& InItem, const FOperatorStackEditorTree& InItemTree, FOperatorStackEditorFooterBuilder& InFooterBuilder)
 {
 	// Customize stack and modifier footer
-	if (InItem->IsA<UActorModifierCoreBase>())
+	if (InItemTree.GetRootItem()->GetValueCount() == 1 && InItem->IsA<UActorModifierCoreBase>())
 	{
-		const UActorModifierCoreBase* Modifier = InItem->Get<UActorModifierCoreBase>();
+		const UActorModifierCoreBase* Modifier = InItem->Get<UActorModifierCoreBase>(0);
 		const TSharedPtr<FActorModifierCoreProfiler> Profiler = Modifier->GetProfiler();
 		UActorModifierCoreEditorSubsystem* ExtensionSubsystem = UActorModifierCoreEditorSubsystem::Get();
 
@@ -281,7 +340,7 @@ void UActorModifierCoreEditorStackCustomization::CustomizeItemFooter(const FOper
 
 bool UActorModifierCoreEditorStackCustomization::OnIsItemDraggable(const FOperatorStackEditorItemPtr& InDragItem)
 {
-	if (InDragItem->IsA<UActorModifierCoreBase>())
+	if (InDragItem->IsA<UActorModifierCoreBase>() && InDragItem->GetValueCount() == 1)
 	{
 		const bool bIsStack = InDragItem->IsA<UActorModifierCoreStack>();
 		return !bIsStack;
@@ -292,22 +351,23 @@ bool UActorModifierCoreEditorStackCustomization::OnIsItemDraggable(const FOperat
 
 TOptional<EItemDropZone> UActorModifierCoreEditorStackCustomization::OnItemCanAcceptDrop(const TArray<FOperatorStackEditorItemPtr>& InDraggedItems, const FOperatorStackEditorItemPtr& InDropZoneItem, EItemDropZone InZone)
 {
-	const UActorModifierCoreSubsystem* ModifierSubsystem = UActorModifierCoreSubsystem::Get();
+	if (!InDropZoneItem->IsA<UActorModifierCoreBase>()
+		|| InDropZoneItem->GetValueCount() != 1)
+	{
+		return Super::OnItemCanAcceptDrop(InDraggedItems, InDropZoneItem, InZone);
+	}
 
 	TSet<UActorModifierCoreBase*> DraggedModifiers;
 	for (const FOperatorStackEditorItemPtr& Item : InDraggedItems)
 	{
-		if (!Item.IsValid())
+		if (!Item.IsValid()
+			|| !Item->IsA<UActorModifierCoreBase>()
+			|| Item->GetValueCount() != 1)
 		{
 			continue;
 		}
 
-		if (!Item->IsA<UActorModifierCoreBase>())
-		{
-			continue;
-		}
-
-		if (UActorModifierCoreBase* Modifier = Item->Get<UActorModifierCoreBase>())
+		if (UActorModifierCoreBase* Modifier = Item->Get<UActorModifierCoreBase>(0))
 		{
 			if (!Modifier->IsModifierStack())
 			{
@@ -316,20 +376,21 @@ TOptional<EItemDropZone> UActorModifierCoreEditorStackCustomization::OnItemCanAc
 		}
 	}
 
-	if (InDropZoneItem->IsA<UActorModifierCoreBase>())
-	{
-		UActorModifierCoreBase* DropModifier = InDropZoneItem->Get<UActorModifierCoreBase>();
+	const UActorModifierCoreSubsystem* ModifierSubsystem = UActorModifierCoreSubsystem::Get();
+	UActorModifierCoreBase* DropModifier = InDropZoneItem->Get<UActorModifierCoreBase>(0);
 
+	if (IsValid(ModifierSubsystem) && IsValid(DropModifier->GetModifiedActor()))
+	{
 		TArray<UActorModifierCoreBase*> MoveModifiers;
-		TArray<UActorModifierCoreBase*> CloneModifiers;
+        TArray<UActorModifierCoreBase*> CloneModifiers;
+        const EActorModifierCoreStackPosition Position = InZone == EItemDropZone::AboveItem ? EActorModifierCoreStackPosition::Before : EActorModifierCoreStackPosition::After;
 
-		const EActorModifierCoreStackPosition Position = InZone == EItemDropZone::AboveItem ? EActorModifierCoreStackPosition::Before : EActorModifierCoreStackPosition::After;
-		ModifierSubsystem->GetSortedModifiers(DraggedModifiers, DropModifier->GetModifiedActor(), DropModifier, Position, MoveModifiers, CloneModifiers);
+        ModifierSubsystem->GetSortedModifiers(DraggedModifiers, DropModifier->GetModifiedActor(), DropModifier, Position, MoveModifiers, CloneModifiers);
 
-		if (!MoveModifiers.IsEmpty())
-		{
-			return InZone;
-		}
+        if (!MoveModifiers.IsEmpty())
+        {
+        	return InZone;
+        }
 	}
 
 	return Super::OnItemCanAcceptDrop(InDraggedItems, InDropZoneItem, InZone);
@@ -337,70 +398,65 @@ TOptional<EItemDropZone> UActorModifierCoreEditorStackCustomization::OnItemCanAc
 
 void UActorModifierCoreEditorStackCustomization::OnDropItem(const TArray<FOperatorStackEditorItemPtr>& InDraggedItems, const FOperatorStackEditorItemPtr& InDropZoneItem, EItemDropZone InZone)
 {
-	const UActorModifierCoreSubsystem* const ModifierSubsystem = UActorModifierCoreSubsystem::Get();
+	if (!InDropZoneItem->IsA<UActorModifierCoreBase>()
+		|| InDropZoneItem->GetValueCount() != 1)
+	{
+		return;
+	}
 
-	TSet<UActorModifierCoreBase*> Modifiers;
+	TSet<UActorModifierCoreBase*> DraggedModifiers;
 	for (const FOperatorStackEditorItemPtr& Item : InDraggedItems)
 	{
-		if (!Item.IsValid())
-		{
-			continue;
-		}
-
-		if (!Item->IsA<UActorModifierCoreBase>())
+		if (!Item.IsValid()
+			|| !Item->IsA<UActorModifierCoreBase>()
+			|| Item->GetValueCount() != 1)
 		{
 			continue;
 		}
 
-		if (UActorModifierCoreBase* Modifier = Item->Get<UActorModifierCoreBase>())
+		if (UActorModifierCoreBase* Modifier = Item->Get<UActorModifierCoreBase>(0))
 		{
 			if (!Modifier->IsModifierStack())
 			{
-				Modifiers.Add(Modifier);
+				DraggedModifiers.Add(Modifier);
 			}
 		}
 	}
 
-	if (InDropZoneItem->IsA<UActorModifierCoreBase>())
-	{
-		UActorModifierCoreBase* DropModifier = InDropZoneItem->Get<UActorModifierCoreBase>();
-
-		UActorModifierCoreStack* const TargetStack = DropModifier->GetModifierStack();
+	const UActorModifierCoreSubsystem* const ModifierSubsystem = UActorModifierCoreSubsystem::Get();
+	UActorModifierCoreBase* DropModifier = InDropZoneItem->Get<UActorModifierCoreBase>(0);
 
-		if (!IsValid(ModifierSubsystem) || !IsValid(TargetStack))
-		{
-			return;
-		}
+	if (!IsValid(ModifierSubsystem) || !IsValid(DropModifier->GetModifiedActor()))
+	{
+		return;
+	}
 
-		AActor* const TargetActor = TargetStack->GetModifiedActor();
+	TArray<UActorModifierCoreBase*> MoveModifiers;
+	TArray<UActorModifierCoreBase*> CloneModifiers;
+	const EActorModifierCoreStackPosition Position = InZone == EItemDropZone::AboveItem ? EActorModifierCoreStackPosition::Before : EActorModifierCoreStackPosition::After;
 
-		TArray<UActorModifierCoreBase*> MoveModifiers;
-		TArray<UActorModifierCoreBase*> CloneModifiers;
+	ModifierSubsystem->GetSortedModifiers(DraggedModifiers, DropModifier->GetModifiedActor(), DropModifier, Position, MoveModifiers, CloneModifiers);
 
-		const EActorModifierCoreStackPosition Position = InZone == EItemDropZone::AboveItem ? EActorModifierCoreStackPosition::Before : EActorModifierCoreStackPosition::After;
-		ModifierSubsystem->GetSortedModifiers(Modifiers, TargetActor, DropModifier, Position, MoveModifiers, CloneModifiers);
-
-		if (MoveModifiers.IsEmpty())
-		{
-			return;
-		}
+	if (MoveModifiers.IsEmpty())
+	{
+		return;
+	}
 
-		FText FailReason;
-		FActorModifierCoreStackMoveOp MoveOp;
-		MoveOp.bShouldTransact = true;
-		MoveOp.FailReason = &FailReason;
-		MoveOp.MovePosition = Position;
-		MoveOp.MovePositionContext = DropModifier;
+	FText FailReason;
+	FActorModifierCoreStackMoveOp MoveOp;
+	MoveOp.bShouldTransact = true;
+	MoveOp.FailReason = &FailReason;
+	MoveOp.MovePosition = Position;
+	MoveOp.MovePositionContext = DropModifier;
 
-		ModifierSubsystem->MoveModifiers(MoveModifiers, TargetStack, MoveOp);
+	ModifierSubsystem->MoveModifiers(MoveModifiers, DropModifier->GetModifierStack(), MoveOp);
 
-		if (!FailReason.IsEmpty())
-		{
-			FNotificationInfo NotificationInfo(FailReason);
-			NotificationInfo.ExpireDuration = 3.f;
-			NotificationInfo.bFireAndForget = true;
-			FSlateNotificationManager::Get().AddNotification(NotificationInfo);
-		}
+	if (!FailReason.IsEmpty())
+	{
+		FNotificationInfo NotificationInfo(FailReason);
+		NotificationInfo.ExpireDuration = 3.f;
+		NotificationInfo.bFireAndForget = true;
+		FSlateNotificationManager::Get().AddNotification(NotificationInfo);
 	}
 
 	Super::OnDropItem(InDraggedItems, InDropZoneItem, InZone);
@@ -436,7 +492,10 @@ void UActorModifierCoreEditorStackCustomization::FillStackHeaderMenu(UToolMenu*
 	{
 		if (ContextItem->IsA<UObject>())
 		{
-			ContextObjects.Add(ContextItem->Get<UObject>());
+			for (UObject* ContextObject : ContextItem->GetAsArray<UObject>())
+			{
+				ContextObjects.Add(ContextObject);
+			}
 		}
 	}
 
@@ -465,21 +524,15 @@ void UActorModifierCoreEditorStackCustomization::FillItemHeaderActionMenu(UToolM
 		return;
 	}
 
-	UActorModifierCoreBase* Modifier = ItemContext->Get<UActorModifierCoreBase>();
-	if (!IsValid(Modifier))
-	{
-		return;
-	}
-
 	// Add profiling stat toggle entry
-	if (UActorModifierCoreStack* ModifierStack = Cast<UActorModifierCoreStack>(Modifier))
+	if (ItemContext->IsA<UActorModifierCoreStack>())
 	{
 		const FToolMenuEntry EnableProfilingModifierAction = FToolMenuEntry::InitToolBarButton(
 			TEXT("EnableProfilingModifierMenuEntry")
 			, FUIAction(
-				FExecuteAction::CreateUObject(this, &UActorModifierCoreEditorStackCustomization::ToggleModifierProfilingAction, ModifierStack)
+				FExecuteAction::CreateUObject(this, &UActorModifierCoreEditorStackCustomization::ToggleModifierProfilingAction, ItemContext)
 				, FCanExecuteAction()
-				, FIsActionChecked::CreateUObject(this, &UActorModifierCoreEditorStackCustomization::IsModifierProfiling, ModifierStack))
+				, FIsActionChecked::CreateUObject(this, &UActorModifierCoreEditorStackCustomization::IsModifierProfiling, ItemContext))
 			, FText::GetEmpty()
 			, FText::GetEmpty()
 			, FSlateIcon(FActorModifierCoreEditorStyle::Get().GetStyleSetName(), "Profiling")
@@ -492,7 +545,10 @@ void UActorModifierCoreEditorStackCustomization::FillItemHeaderActionMenu(UToolM
 	// Add remove modifier entry
 	const FToolMenuEntry RemoveModifierAction = FToolMenuEntry::InitToolBarButton(
 		TEXT("RemoveModifierMenuEntry")
-		, FExecuteAction::CreateUObject(this, &UActorModifierCoreEditorStackCustomization::RemoveModifierAction, Modifier)
+		, FUIAction(
+			FExecuteAction::CreateUObject(this, &UActorModifierCoreEditorStackCustomization::RemoveModifierAction, ItemContext)
+			, FCanExecuteAction::CreateUObject(this, &UActorModifierCoreEditorStackCustomization::CanRemoveModifier, ItemContext)
+		  )
 		, FText::GetEmpty()
 		, FText::GetEmpty()
 		, FSlateIcon(FAppStyle::GetAppStyleSetName(), "Icons.Delete")
@@ -520,13 +576,7 @@ void UActorModifierCoreEditorStackCustomization::FillItemContextActionMenu(UTool
 		return;
 	}
 
-	UActorModifierCoreBase* Modifier = ItemContext->Get<UActorModifierCoreBase>();
-	if (!IsValid(Modifier))
-	{
-		return;
-	}
-
-	if (UActorModifierCoreStack* ModifierStack = Cast<UActorModifierCoreStack>(Modifier))
+	if (ItemContext->IsA<UActorModifierCoreStack>())
 	{
 		const FToolMenuEntry EnableProfilingModifierAction = FToolMenuEntry::InitMenuEntry(
 			TEXT("EnableProfilingModifierMenuEntry")
@@ -534,9 +584,9 @@ void UActorModifierCoreEditorStackCustomization::FillItemContextActionMenu(UTool
 			, FText::GetEmpty()
 			, FSlateIcon(FActorModifierCoreEditorStyle::Get().GetStyleSetName(), "Profiling")
 			, FUIAction(
-				FExecuteAction::CreateUObject(this, &UActorModifierCoreEditorStackCustomization::ToggleModifierProfilingAction, ModifierStack)
+				FExecuteAction::CreateUObject(this, &UActorModifierCoreEditorStackCustomization::ToggleModifierProfilingAction, ItemContext)
 				, FCanExecuteAction()
-				, FIsActionChecked::CreateUObject(this, &UActorModifierCoreEditorStackCustomization::IsModifierProfiling, ModifierStack))
+				, FIsActionChecked::CreateUObject(this, &UActorModifierCoreEditorStackCustomization::IsModifierProfiling, ItemContext))
 			, EUserInterfaceActionType::ToggleButton);
 
 		InToolMenu->AddMenuEntry(EnableProfilingModifierAction.Name, EnableProfilingModifierAction);
@@ -570,9 +620,9 @@ void UActorModifierCoreEditorStackCustomization::FillItemContextActionMenu(UTool
 	}
 }
 
-bool UActorModifierCoreEditorStackCustomization::CanRemoveModifier(UActorModifierCoreBase* InModifier) const
+bool UActorModifierCoreEditorStackCustomization::CanRemoveModifier(FOperatorStackEditorItemPtr InItem) const
 {
-	if (!IsValid(InModifier) || !IsValid(InModifier->GetModifiedActor()))
+	if (!InItem.IsValid() || !InItem->HasValue())
 	{
 		return false;
 	}
@@ -580,9 +630,9 @@ bool UActorModifierCoreEditorStackCustomization::CanRemoveModifier(UActorModifie
 	return true;
 }
 
-void UActorModifierCoreEditorStackCustomization::RemoveModifierAction(UActorModifierCoreBase* InModifier) const
+void UActorModifierCoreEditorStackCustomization::RemoveModifierAction(FOperatorStackEditorItemPtr InItem) const
 {
-	if (!CanRemoveModifier(InModifier))
+	if (!CanRemoveModifier(InItem))
 	{
 		return;
 	}
@@ -594,19 +644,24 @@ void UActorModifierCoreEditorStackCustomization::RemoveModifierAction(UActorModi
 		return;
 	}
 
-	if (InModifier->IsModifierStack())
+	if (InItem->IsA<UActorModifierCoreStack>())
 	{
-		TSet<AActor*> Actors
+		TSet<AActor*> Actors;
+
+		for (const UActorModifierCoreStack* ModifierStack : InItem->GetAsArray<UActorModifierCoreStack>())
 		{
-			InModifier->GetModifiedActor()
-		};
+			Actors.Add(ModifierStack->GetModifiedActor());
+		}
 
-		ModifierSubsystem->RemoveActorsModifiers(Actors, true);
+		if (!ModifierSubsystem->RemoveActorsModifiers(Actors, true))
+		{
+			UE_LOG(LogActorModifierCoreEditorStackCustomization, Warning, TEXT("Could not remove modifiers from actors"))
+		}
 	}
 	else
 	{
 		TSet<UActorModifierCoreBase*> Modifiers;
-		Modifiers.Add(InModifier);
+		Modifiers.Append(InItem->GetAsArray<UActorModifierCoreBase>());
 
 		FText OutFailReason;
 		FActorModifierCoreStackRemoveOp RemoveOp;
@@ -623,9 +678,9 @@ void UActorModifierCoreEditorStackCustomization::RemoveModifierAction(UActorModi
 	}
 }
 
-bool UActorModifierCoreEditorStackCustomization::CanCopyModifier(UActorModifierCoreBase* InModifier) const
+bool UActorModifierCoreEditorStackCustomization::CanCopyModifier(FOperatorStackEditorItemPtr InItem) const
 {
-	if (!IsValid(InModifier) || InModifier->IsModifierStack())
+	if (!InItem.IsValid() || !InItem->HasValue() || InItem->GetValueCount() != 1)
 	{
 		return false;
 	}
@@ -633,40 +688,41 @@ bool UActorModifierCoreEditorStackCustomization::CanCopyModifier(UActorModifierC
 	return true;
 }
 
-void UActorModifierCoreEditorStackCustomization::CopyModifierAction(UActorModifierCoreBase* InModifier) const
+void UActorModifierCoreEditorStackCustomization::CopyModifierAction(FOperatorStackEditorItemPtr InItem) const
 {
-	if (!CanCopyModifier(InModifier))
+	if (!CanCopyModifier(InItem))
 	{
 		return;
 	}
 
-	TMap<FName, FString> ModifierPropertiesHandlesMap;
-	const bool bSuccess = CreatePropertiesHandlesMapFromModifier(InModifier, ModifierPropertiesHandlesMap);
-
-	if (!bSuccess)
+	// Should only contain one modifier since we only allow action for one
+	TArray<TSharedPtr<FJsonValue>> JsonModifiers;
+	for (UActorModifierCoreBase* Modifier : InItem->GetAsArray<UActorModifierCoreBase>())
 	{
-		return;
-	}
-
-	FString SerializedString;
-
-	const TSharedRef<FJsonObject> RootJsonObject = MakeShared<FJsonObject>();
-
-	FActorModifierCoreEditorPropertiesWrapper ModifierPropertiesWrapper;
-	ModifierPropertiesWrapper.ModifierName = InModifier->GetModifierName();
-	ModifierPropertiesWrapper.PropertiesHandlesAsStringMap = ModifierPropertiesHandlesMap;
+		TMap<FName, FString> ModifierPropertiesHandlesMap;
+		if (!CreatePropertiesHandlesMapFromModifier(Modifier, ModifierPropertiesHandlesMap))
+		{
+			continue;
+		}
 
-	TSharedRef<FJsonObject> PropertiesJsonObject = MakeShared<FJsonObject>();
-	FJsonObjectConverter::UStructToJsonObject(FActorModifierCoreEditorPropertiesWrapper::StaticStruct(), &ModifierPropertiesWrapper, PropertiesJsonObject, 0 /* CheckFlags */, 0 /* SkipFlags */);
-	const TSharedPtr<FJsonValue> PropertiesJsonValue = MakeShared<FJsonValueObject>(PropertiesJsonObject);
+		FActorModifierCoreEditorPropertiesWrapper ModifierPropertiesWrapper;
+		ModifierPropertiesWrapper.ModifierName = Modifier->GetModifierName();
+		ModifierPropertiesWrapper.PropertiesHandlesAsStringMap = ModifierPropertiesHandlesMap;
 
-	RootJsonObject->SetField(PropertiesWrapperEntry, PropertiesJsonValue);
+		TSharedRef<FJsonObject> PropertiesJsonObject = MakeShared<FJsonObject>();
+		FJsonObjectConverter::UStructToJsonObject(FActorModifierCoreEditorPropertiesWrapper::StaticStruct(), &ModifierPropertiesWrapper, PropertiesJsonObject, 0 /* CheckFlags */, 0 /* SkipFlags */);
+		JsonModifiers.Add(MakeShared<FJsonValueObject>(PropertiesJsonObject));
+	}
 
+	FString SerializedString;
 	const TSharedRef<TJsonWriter<>> Writer = TJsonWriterFactory<>::Create(&SerializedString);
+	if (!FJsonSerializer::Serialize(JsonModifiers, Writer))
+	{
+		UE_LOG(LogActorModifierCoreEditorStackCustomization, Warning, TEXT("Unable to serialize the selected modifier(s) into Json format"));
+		return;
+	}
 
-	FJsonSerializer::Serialize(RootJsonObject, Writer);
-
-	// Add Prefix to quickly identify whether current clipboard is from Tag Collection Properties or not
+	// Add prefix to quickly identify whether current clipboard is from modifiers or not
 	SerializedString = *FString::Printf(TEXT("%s%s")
 		, *PropertiesWrapperPrefix
 		, *SerializedString);
@@ -674,97 +730,141 @@ void UActorModifierCoreEditorStackCustomization::CopyModifierAction(UActorModifi
 	FPlatformApplicationMisc::ClipboardCopy(*SerializedString);
 }
 
-bool UActorModifierCoreEditorStackCustomization::CanPasteModifier(UActorModifierCoreBase* InModifier) const
+bool UActorModifierCoreEditorStackCustomization::CanPasteModifier(FOperatorStackEditorItemPtr InItem) const
 {
-	if (!IsValid(InModifier))
+	if (!InItem.IsValid() || !InItem->HasValue() || InItem->GetValueCount() != 1)
 	{
 		return false;
 	}
 
 	FString ClipboardContent;
 	FPlatformApplicationMisc::ClipboardPaste(ClipboardContent);
+
 	return ClipboardContent.StartsWith(PropertiesWrapperPrefix);
 }
 
-void UActorModifierCoreEditorStackCustomization::PasteModifierAction(UActorModifierCoreBase* InModifier) const
+void UActorModifierCoreEditorStackCustomization::PasteModifierAction(FOperatorStackEditorItemPtr InItem) const
 {
-	if (!CanPasteModifier(InModifier))
+	if (!CanPasteModifier(InItem))
+	{
+		return;
+	}
+
+	// We only allow usage of this action when one item is selected but we could support multi selection later
+	TArray<FActorModifierCoreEditorPropertiesWrapper> ModifierPropertiesWrapper;
+	if (!GetModifierPropertiesWrapperFromClipboard(ModifierPropertiesWrapper))
 	{
 		return;
 	}
 
-	if (const UActorModifierCoreStack* ModifierStack = Cast<UActorModifierCoreStack>(InModifier))
+	const FName ModifierName = ModifierPropertiesWrapper[0].ModifierName;
+	TArray<UActorModifierCoreBase*> TargetModifiers;
+
+	if (InItem->IsA<UActorModifierCoreStack>())
 	{
-		TSet<AActor*> Actors { InModifier->GetModifiedActor() };
-		if (!AddModifierFromClipboard(Actors))
+		TSet<AActor*> Actors;
+		for (const UActorModifierCoreStack* ModifierStack : InItem->GetAsArray<UActorModifierCoreStack>())
+		{
+			Actors.Add(ModifierStack->GetModifiedActor());
+		}
+
+		if (!AddModifierFromClipboard(Actors, ModifierName, TargetModifiers))
 		{
 			return;
 		}
-		InModifier = ModifierStack->GetLastModifier();
+	}
+	else
+	{
+		TargetModifiers = InItem->GetAsArray<UActorModifierCoreBase>();
 	}
 
-	FActorModifierCoreEditorPropertiesWrapper ModifierPropertiesWrapper;
-	if (GetModifierPropertiesWrapperFromClipboard(ModifierPropertiesWrapper))
+	if (TargetModifiers.Num() != ModifierPropertiesWrapper.Num())
 	{
-		if (ModifierPropertiesWrapper.ModifierName == InModifier->GetModifierName())
+		UE_LOG(LogActorModifierCoreEditorStackCustomization, Warning, TEXT("Unable to set properties from %s modifier due to target modifiers (%i) and modifiers properties (%i) count mismatch"), *ModifierName.ToString(), TargetModifiers.Num(), ModifierPropertiesWrapper.Num());
+		return;
+	}
+
+	FScopedTransaction Transaction(LOCTEXT("PasteModifierProperties", "Paste Modifier Properties"));
+
+	for (int32 Index = 0; Index < TargetModifiers.Num(); Index++)
+	{
+		UActorModifierCoreBase* TargetModifier = TargetModifiers[Index];
+
+		if (!IsValid(TargetModifier))
 		{
-			FScopedTransaction Transaction(LOCTEXT("PasteModifierProperties", "Paste Modifier Properties"));
-			InModifier->Modify();
+			continue;
+		}
 
-			UpdateModifierFromPropertiesHandlesMap(InModifier, ModifierPropertiesWrapper.PropertiesHandlesAsStringMap);
+		const FActorModifierCoreEditorPropertiesWrapper ModifierProperties = ModifierPropertiesWrapper[Index];
+
+		if (ModifierName == TargetModifier->GetModifierName())
+		{
+			TargetModifier->Modify();
+			UpdateModifierFromPropertiesHandlesMap(TargetModifier, ModifierProperties.PropertiesHandlesAsStringMap);
 		}
 		else
 		{
-			const FString SourceModifierName = ModifierPropertiesWrapper.ModifierName.ToString();
-			const FString DestinationModifierName = InModifier->GetModifierName().ToString();
+			const FString SourceModifierName = ModifierProperties.ModifierName.ToString();
+			const FString DestinationModifierName = TargetModifier->GetModifierName().ToString();
 
-			UE_LOG(LogActorModifierCoreEditorPropertiesWrapper, Warning, TEXT("Unable to copy properties from %s modifier to %s modifier"), *SourceModifierName, *DestinationModifierName);
+			UE_LOG(LogActorModifierCoreEditorStackCustomization, Warning, TEXT("Unable to copy properties from %s modifier to %s modifier"), *SourceModifierName, *DestinationModifierName);
 		}
 	}
 }
 
-bool UActorModifierCoreEditorStackCustomization::IsModifierProfiling(UActorModifierCoreStack* InStack) const
+bool UActorModifierCoreEditorStackCustomization::IsModifierProfiling(FOperatorStackEditorItemPtr InItem) const
 {
-	if (!IsValid(InStack) || !IsValid(InStack->GetModifiedActor()))
+	if (!InItem.IsValid() || !InItem->HasValue())
 	{
 		return false;
 	}
 
-	return InStack->IsModifierProfiling();
+	for (const UActorModifierCoreStack* ModifierStack : InItem->GetAsArray<UActorModifierCoreStack>())
+	{
+		if (!ModifierStack->IsModifierProfiling())
+		{
+			return false;
+		}
+	}
+
+	return true;
 }
 
-void UActorModifierCoreEditorStackCustomization::ToggleModifierProfilingAction(UActorModifierCoreStack* InStack) const
+void UActorModifierCoreEditorStackCustomization::ToggleModifierProfilingAction(FOperatorStackEditorItemPtr InItem) const
 {
-	if (!IsValid(InStack) || !IsValid(InStack->GetModifiedActor()))
+	if (!InItem.IsValid() || !InItem->HasValue())
 	{
 		return;
 	}
 
-	InStack->SetModifierProfiling(!InStack->IsModifierProfiling());
+	for (UActorModifierCoreStack* ModifierStack : InItem->GetAsArray<UActorModifierCoreStack>())
+	{
+		ModifierStack->SetModifierProfiling(!ModifierStack->IsModifierProfiling());
+	}
 }
 
-TSharedRef<FUICommandList> UActorModifierCoreEditorStackCustomization::CreateModifierCommands(UActorModifierCoreBase* InModifier)
+TSharedRef<FUICommandList> UActorModifierCoreEditorStackCustomization::CreateModifierCommands(FOperatorStackEditorItemPtr InItem)
 {
 	TSharedRef<FUICommandList> Commands = MakeShared<FUICommandList>();
 
 	Commands->MapAction(FGenericCommands::Get().Copy
 		, FUIAction(
-			FExecuteAction::CreateUObject(this, &UActorModifierCoreEditorStackCustomization::CopyModifierAction, InModifier),
-			FCanExecuteAction::CreateUObject(this, &UActorModifierCoreEditorStackCustomization::CanCopyModifier, InModifier)
+			FExecuteAction::CreateUObject(this, &UActorModifierCoreEditorStackCustomization::CopyModifierAction, InItem),
+			FCanExecuteAction::CreateUObject(this, &UActorModifierCoreEditorStackCustomization::CanCopyModifier, InItem)
 		)
 	);
 
 	Commands->MapAction(FGenericCommands::Get().Paste
 		, FUIAction(
-			FExecuteAction::CreateUObject(this, &UActorModifierCoreEditorStackCustomization::PasteModifierAction, InModifier),
-			FCanExecuteAction::CreateUObject(this, &UActorModifierCoreEditorStackCustomization::CanPasteModifier, InModifier)
+			FExecuteAction::CreateUObject(this, &UActorModifierCoreEditorStackCustomization::PasteModifierAction, InItem),
+			FCanExecuteAction::CreateUObject(this, &UActorModifierCoreEditorStackCustomization::CanPasteModifier, InItem)
 		)
 	);
 
 	Commands->MapAction(FGenericCommands::Get().Delete
 		, FUIAction(
-			FExecuteAction::CreateUObject(this, &UActorModifierCoreEditorStackCustomization::RemoveModifierAction, InModifier),
-			FCanExecuteAction::CreateUObject(this, &UActorModifierCoreEditorStackCustomization::CanRemoveModifier, InModifier)
+			FExecuteAction::CreateUObject(this, &UActorModifierCoreEditorStackCustomization::RemoveModifierAction, InItem),
+			FCanExecuteAction::CreateUObject(this, &UActorModifierCoreEditorStackCustomization::CanRemoveModifier, InItem)
 		)
 	);
 
@@ -814,43 +914,41 @@ bool UActorModifierCoreEditorStackCustomization::CreatePropertiesHandlesMapFromM
 	return true;
 }
 
-bool UActorModifierCoreEditorStackCustomization::GetModifierPropertiesWrapperFromClipboard(FActorModifierCoreEditorPropertiesWrapper& OutPropertiesWrapper) const
+bool UActorModifierCoreEditorStackCustomization::GetModifierPropertiesWrapperFromClipboard(TArray<FActorModifierCoreEditorPropertiesWrapper>& OutModifierProperties) const
 {
 	FString ClipboardContent;
 	FPlatformApplicationMisc::ClipboardPaste(ClipboardContent);
 
+	if (!ClipboardContent.StartsWith(PropertiesWrapperPrefix))
+	{
+		return false;
+	}
+
 	// remove prefix, this is not part of the modifier properties json data
 	ClipboardContent.RightChopInline(PropertiesWrapperPrefix.Len());
 
-	TSharedPtr<FJsonObject> RootJsonObject;
+	TArray<TSharedPtr<FJsonValue>> JsonModifiers;
 	const TSharedRef<TJsonReader<>> Reader = TJsonReaderFactory<>::Create(ClipboardContent);
 
-	if (!FJsonSerializer::Deserialize(Reader, RootJsonObject))
+	if (!FJsonSerializer::Deserialize(Reader, JsonModifiers))
 	{
-		UE_LOG(LogActorModifierCoreEditorPropertiesWrapper, Warning, TEXT("Unable to serialize the pasted text into Json format"));
+		UE_LOG(LogActorModifierCoreEditorStackCustomization, Warning, TEXT("Unable to deserialize the clipboard text into Json format"));
 		return false;
 	}
 
-	const TSharedPtr<FJsonValue> PropertiesJsonValue = RootJsonObject->TryGetField(PropertiesWrapperEntry);
-
-	if (!PropertiesJsonValue.IsValid())
+	if (JsonModifiers.IsEmpty())
 	{
-		UE_LOG(LogActorModifierCoreEditorPropertiesWrapper, Warning, TEXT("Missing %s entry field in pasted text"), *PropertiesWrapperEntry);
+		UE_LOG(LogActorModifierCoreEditorStackCustomization, Warning, TEXT("No Json modifiers data available to paste"));
 		return false;
 	}
 
-	if (PropertiesJsonValue->Type != EJson::Object)
+	if (!FJsonObjectConverter::JsonArrayToUStruct(JsonModifiers, &OutModifierProperties))
 	{
-		UE_LOG(LogActorModifierCoreEditorPropertiesWrapper, Warning, TEXT("Invalid Modifier Properties Json Value. Not an object"));
+		UE_LOG(LogActorModifierCoreEditorStackCustomization, Warning, TEXT("Invalid Json modifiers properties found"));
 		return false;
 	}
 
-	const TSharedPtr<FJsonObject>& PropertiesJsonObject = PropertiesJsonValue->AsObject();
-	check(PropertiesJsonObject.IsValid());
-
-	const bool bStructSuccess = FJsonObjectConverter::JsonObjectToUStruct(PropertiesJsonObject.ToSharedRef(), FActorModifierCoreEditorPropertiesWrapper::StaticStruct(), &OutPropertiesWrapper, 0 /* CheckFlags */, 0 /* SkipFlags */);
-
-	return bStructSuccess;
+	return !OutModifierProperties.IsEmpty();
 }
 
 bool UActorModifierCoreEditorStackCustomization::UpdateModifierFromPropertiesHandlesMap(UActorModifierCoreBase* InModifier, const TMap<FName, FString>& InModifierPropertiesHandlesMap) const
@@ -896,7 +994,7 @@ bool UActorModifierCoreEditorStackCustomization::UpdateModifierFromPropertiesHan
 	return true;
 }
 
-bool UActorModifierCoreEditorStackCustomization::AddModifierFromClipboard(TSet<AActor*>& InActors) const
+bool UActorModifierCoreEditorStackCustomization::AddModifierFromClipboard(const TSet<AActor*>& InActors, FName InModifierName, TArray<UActorModifierCoreBase*>& OutNewModifiers) const
 {
 	const UActorModifierCoreSubsystem* ModifierSubsystem = UActorModifierCoreSubsystem::Get();
 	if (!IsValid(ModifierSubsystem))
@@ -904,31 +1002,25 @@ bool UActorModifierCoreEditorStackCustomization::AddModifierFromClipboard(TSet<A
 		return false;
 	}
 
-	FActorModifierCoreEditorPropertiesWrapper ModifierPropertiesWrapper;
-	if (GetModifierPropertiesWrapperFromClipboard(ModifierPropertiesWrapper))
-	{
-		const FName ModifierName = ModifierPropertiesWrapper.ModifierName;
-
-		FText OutFailReason;
-		FActorModifierCoreStackInsertOp AddOp;
-		AddOp.bShouldTransact = true;
-		AddOp.FailReason = &OutFailReason;
-		AddOp.NewModifierName = ModifierName;
+	FText OutFailReason;
+	FActorModifierCoreStackInsertOp AddOp;
+	AddOp.bShouldTransact = true;
+	AddOp.FailReason = &OutFailReason;
+	AddOp.NewModifierName = InModifierName;
 
-		ModifierSubsystem->AddActorsModifiers(InActors, AddOp);
+	OutNewModifiers = ModifierSubsystem->AddActorsModifiers(InActors, AddOp);
 
-		if (!OutFailReason.IsEmpty())
-		{
-			FNotificationInfo NotificationInfo(OutFailReason);
-			NotificationInfo.ExpireDuration = 3.0f;
-			NotificationInfo.bFireAndForget = true;
-			FSlateNotificationManager::Get().AddNotification(NotificationInfo);
-		}
+	if (!OutFailReason.IsEmpty())
+	{
+		FNotificationInfo NotificationInfo(OutFailReason);
+		NotificationInfo.ExpireDuration = 3.0f;
+		NotificationInfo.bFireAndForget = true;
+		FSlateNotificationManager::Get().AddNotification(NotificationInfo);
 
-		return true;
+		return false;
 	}
 
-	return false;
+	return !OutNewModifiers.IsEmpty();
 }
 
 #undef LOCTEXT_NAMESPACE
diff --git a/Engine/Plugins/Experimental/ActorModifierCore/Source/ActorModifierCoreEditor/Private/Modifiers/Customizations/ActorModifierCoreEditorStackCustomization.h b/Engine/Plugins/Experimental/ActorModifierCore/Source/ActorModifierCoreEditor/Private/Modifiers/Customizations/ActorModifierCoreEditorStackCustomization.h
--- a/Engine/Plugins/Experimental/ActorModifierCore/Source/ActorModifierCoreEditor/Private/Modifiers/Customizations/ActorModifierCoreEditorStackCustomization.h
+++ b/Engine/Plugins/Experimental/ActorModifierCore/Source/ActorModifierCoreEditor/Private/Modifiers/Customizations/ActorModifierCoreEditorStackCustomization.h
@@ -32,12 +32,12 @@ class UActorModifierCoreEditorStackCustomization : public UOperatorStackEditorSt
 
 public:
 	static inline const FString PropertiesWrapperPrefix = TEXT("ActorModifierCoreEditorPropertiesWrapper");
-	static inline const FString PropertiesWrapperEntry  = TEXT("PropertiesWrapper");
 
 	UActorModifierCoreEditorStackCustomization();
 
 	//~ Begin UOperatorStackEditorStackCustomization
-	virtual bool TransformContextItem(const FOperatorStackEditorItemPtr& InItem, TArray<FOperatorStackEditorItemPtr>& OutTransformedItems) const override;
+	virtual bool GetRootItem(const FOperatorStackEditorContext& InContext, FOperatorStackEditorItemPtr& OutRootItem) const override;
+	virtual bool GetChildrenItem(const FOperatorStackEditorItemPtr& InItem, TArray<FOperatorStackEditorItemPtr>& OutChildrenItems) const override;
 	virtual void CustomizeStackHeader(const FOperatorStackEditorTree& InItemTree, FOperatorStackEditorHeaderBuilder& InHeaderBuilder) override;
 	virtual void CustomizeItemHeader(const FOperatorStackEditorItemPtr& InItem, const FOperatorStackEditorTree& InItemTree, FOperatorStackEditorHeaderBuilder& InHeaderBuilder) override;
 	virtual void CustomizeItemBody(const FOperatorStackEditorItemPtr& InItem, const FOperatorStackEditorTree& InItemTree, FOperatorStackEditorBodyBuilder& InBodyBuilder) override;
@@ -58,25 +58,25 @@ protected:
 	void FillItemContextActionMenu(UToolMenu* InToolMenu) const;
 
 	/** Remove modifier action */
-	bool CanRemoveModifier(UActorModifierCoreBase* InModifier) const;
-	void RemoveModifierAction(UActorModifierCoreBase* InModifier) const;
+	bool CanRemoveModifier(FOperatorStackEditorItemPtr InItem) const;
+	void RemoveModifierAction(FOperatorStackEditorItemPtr InItem) const;
 
 	/** Copy modifier action */
-	bool CanCopyModifier(UActorModifierCoreBase* InModifier) const;
-	void CopyModifierAction(UActorModifierCoreBase* InModifier) const;
+	bool CanCopyModifier(FOperatorStackEditorItemPtr InItem) const;
+	void CopyModifierAction(FOperatorStackEditorItemPtr InItem) const;
 
 	/** Paste modifier action */
-	bool CanPasteModifier(UActorModifierCoreBase* InModifier) const;
-	void PasteModifierAction(UActorModifierCoreBase* InModifier) const;
+	bool CanPasteModifier(FOperatorStackEditorItemPtr InItem) const;
+	void PasteModifierAction(FOperatorStackEditorItemPtr InItem) const;
 
 	/** Toggle the profiling for a whole stack */
-	bool IsModifierProfiling(UActorModifierCoreStack* InStack) const;
-	void ToggleModifierProfilingAction(UActorModifierCoreStack* InStack) const;
+	bool IsModifierProfiling(FOperatorStackEditorItemPtr InItem) const;
+	void ToggleModifierProfilingAction(FOperatorStackEditorItemPtr InItem) const;
 
-	TSharedRef<FUICommandList> CreateModifierCommands(UActorModifierCoreBase* InModifier);
+	TSharedRef<FUICommandList> CreateModifierCommands(FOperatorStackEditorItemPtr InItem);
 
 	bool CreatePropertiesHandlesMapFromModifier(UActorModifierCoreBase* InModifier, TMap<FName, FString>& OutModifierPropertiesHandlesMap) const;
-	bool GetModifierPropertiesWrapperFromClipboard(FActorModifierCoreEditorPropertiesWrapper& OutPropertiesWrapper) const;
+	bool GetModifierPropertiesWrapperFromClipboard(TArray<FActorModifierCoreEditorPropertiesWrapper>& OutModifierProperties) const;
 	bool UpdateModifierFromPropertiesHandlesMap(UActorModifierCoreBase* InModifier, const TMap<FName, FString>& InModifierPropertiesHandlesMap) const;
-	bool AddModifierFromClipboard(TSet<AActor*>& InActors) const;
+	bool AddModifierFromClipboard(const TSet<AActor*>& InActors, FName InModifierName, TArray<UActorModifierCoreBase*>& OutNewModifiers) const;
 };
\ No newline at end of file
diff --git a/Engine/Plugins/Experimental/OperatorStack/Source/OperatorStack/Private/Items/OperatorStackEditorTree.cpp b/Engine/Plugins/Experimental/OperatorStack/Source/OperatorStack/Private/Items/OperatorStackEditorTree.cpp
--- a/Engine/Plugins/Experimental/OperatorStack/Source/OperatorStack/Private/Items/OperatorStackEditorTree.cpp
+++ b/Engine/Plugins/Experimental/OperatorStack/Source/OperatorStack/Private/Items/OperatorStackEditorTree.cpp
@@ -19,24 +19,26 @@ FOperatorStackEditorTree::FOperatorStackEditorTree(UOperatorStackEditorStackCust
 
 	ContextWeak = InContext;
 
-	// Build if context is valid
-	const TArray<FOperatorStackEditorItemPtr> ContextItems(InContext->GetItems());
-	BuildTree(ContextItems);
+	// Get root items
+	TSharedPtr<FOperatorStackEditorItem> RootItem;
+	if (InCustomization->GetRootItem(*InContext.Get(), RootItem))
+	{
+		check(InCustomization->IsCustomizationSupportedFor(RootItem))
+
+		Items.Empty(1);
+
+		BuildTreeInternal({RootItem}, INDEX_NONE);
+	}
 }
 
-TArray<FOperatorStackEditorItemPtr> FOperatorStackEditorTree::GetRootItems() const
+FOperatorStackEditorItemPtr FOperatorStackEditorTree::GetRootItem() const
 {
-	TArray<FOperatorStackEditorItemPtr> RootItems;
-	RootItems.Reserve(RootNodes.Num());
-
-	for (const FOperatorStackEditorTreeNodePtr& Node : RootNodes)
+	if (RootNode.IsValid() && Items.IsValidIndex(RootNode->ItemIndex))
 	{
-		check(Items.IsValidIndex(Node->ItemIndex))
-
-		RootItems.Add(Items[Node->ItemIndex]);
+		return Items[RootNode->ItemIndex];
 	}
 
-	return RootItems;
+	return nullptr;
 }
 
 TArray<FOperatorStackEditorItemPtr> FOperatorStackEditorTree::GetChildrenItems(FOperatorStackEditorItemPtr InItem) const
@@ -129,29 +131,6 @@ bool FOperatorStackEditorTree::Contains(FOperatorStackEditorItemPtr InItem) cons
 	});
 }
 
-void FOperatorStackEditorTree::GetSupportedItems(const TArray<FOperatorStackEditorItemPtr>& InItems, TArray<FOperatorStackEditorItemPtr>& OutSupportedItems) const
-{
-	const UOperatorStackEditorStackCustomization* StackCustomization = CustomizationWeak.Get();
-
-	// Gather supported items for this customization
-	for (const FOperatorStackEditorItemPtr& Item : InItems)
-	{
-		if (!Item.IsValid() || !Item->HasValue())
-		{
-			continue;
-		}
-
-		if (StackCustomization->IsCustomizationSupportedFor(Item))
-		{
-			OutSupportedItems.Add(Item);
-			continue;
-		}
-
-		// Item is not supported, let's transform it to something supported
-		OutSupportedItems.Append(GetSupportedChildrenItems(Item));
-	}
-}
-
 TArray<FOperatorStackEditorItemPtr> FOperatorStackEditorTree::GetSupportedChildrenItems(const FOperatorStackEditorItemPtr& InParentItem) const
 {
 	TArray<FOperatorStackEditorItemPtr> ChildrenPtr;
@@ -162,25 +141,16 @@ TArray<FOperatorStackEditorItemPtr> FOperatorStackEditorTree::GetSupportedChildr
 	}
 
 	const UOperatorStackEditorStackCustomization* StackCustomization = CustomizationWeak.Get();
-	StackCustomization->TransformContextItem(InParentItem, ChildrenPtr);
-	ChildrenPtr.RemoveAll([](const FOperatorStackEditorItemPtr& InItem)
+	StackCustomization->GetChildrenItem(InParentItem, ChildrenPtr);
+
+	ChildrenPtr.RemoveAll([StackCustomization](const FOperatorStackEditorItemPtr& InItem)
 	{
-		return !InItem.IsValid() || !InItem->HasValue();
+		return !InItem.IsValid() || !InItem->HasValue() || !StackCustomization->IsCustomizationSupportedFor(InItem);
 	});
 
 	return ChildrenPtr;
 }
 
-void FOperatorStackEditorTree::BuildTree(const TArray<FOperatorStackEditorItemPtr>& InSourceItems)
-{
-	TArray<FOperatorStackEditorItemPtr> SupportedItems;
-	GetSupportedItems(InSourceItems, SupportedItems);
-
-	Items.Empty(SupportedItems.Num());
-
-	BuildTreeInternal(SupportedItems, INDEX_NONE);
-}
-
 void FOperatorStackEditorTree::BuildTreeInternal(const TArray<FOperatorStackEditorItemPtr>& InSupportedItems, int32 InParentIndex)
 {
 	for (const FOperatorStackEditorItemPtr& SupportedItem : InSupportedItems)
@@ -200,7 +170,8 @@ void FOperatorStackEditorTree::BuildTreeInternal(const TArray<FOperatorStackEdit
 
 		if (InParentIndex == INDEX_NONE)
 		{
-			RootNodes.Add(Node);
+			check(!RootNode.IsValid())
+			RootNode = Node;
 		}
 		else
 		{
diff --git a/Engine/Plugins/Experimental/OperatorStack/Source/OperatorStack/Private/Widgets/SOperatorStackEditorPanel.cpp b/Engine/Plugins/Experimental/OperatorStack/Source/OperatorStack/Private/Widgets/SOperatorStackEditorPanel.cpp
--- a/Engine/Plugins/Experimental/OperatorStack/Source/OperatorStack/Private/Widgets/SOperatorStackEditorPanel.cpp
+++ b/Engine/Plugins/Experimental/OperatorStack/Source/OperatorStack/Private/Widgets/SOperatorStackEditorPanel.cpp
@@ -182,12 +182,12 @@ const FOperatorStackEditorTree& SOperatorStackEditorPanel::GetItemTree(UOperator
 	return *Tree;
 }
 
-void SOperatorStackEditorPanel::SaveItemExpansionState(const void* InItem, bool bInExpanded)
+void SOperatorStackEditorPanel::SaveItemExpansionState(uint32 InItem, bool bInExpanded)
 {
 	ItemExpansionState.Add(InItem, bInExpanded);
 }
 
-bool SOperatorStackEditorPanel::GetItemExpansionState(const void* InItem, bool& bOutExpanded)
+bool SOperatorStackEditorPanel::GetItemExpansionState(uint32 InItem, bool& bOutExpanded)
 {
 	if (const bool* bExpanded = ItemExpansionState.Find(InItem))
 	{
@@ -210,9 +210,6 @@ void SOperatorStackEditorPanel::AddSlot(const UOperatorStackEditorStackCustomiza
 		return;
 	}
 
-	static const FToolBarStyle& ToolBarStyle = FAppStyle::Get().GetWidgetStyle<FToolBarStyle>("SlimToolBar");
-	static const FCheckBoxStyle* const CheckStyle = &ToolBarStyle.ToggleButton;
-
 	const FText& Label = InCustomizationStack->GetLabel();
 	const FName& Identifier = InCustomizationStack->GetIdentifier();
 	const FSlateBrush* Icon = InCustomizationStack->GetIcon();
@@ -235,32 +232,35 @@ void SOperatorStackEditorPanel::AddSlot(const UOperatorStackEditorStackCustomiza
 	WidgetToolbar->AddSlot()
 		.Padding(SOperatorStackEditorStack::Padding)
 		[
-			SNew(SCheckBox)
-			.Style(CheckStyle)
-			.ForegroundColor(FLinearColor::White)
-			.OnCheckStateChanged(this, &SOperatorStackEditorPanel::OnToolbarButtonClicked, WidgetIdx)
-			.IsChecked(this, &SOperatorStackEditorPanel::IsToolbarButtonActive, WidgetIdx)
-			.Visibility(this, &SOperatorStackEditorPanel::GetToolbarButtonVisibility, WidgetIdx)
+			SNew(SBox)
+			.Padding(FMargin(0.0f))
 			[
-				SNew(SHorizontalBox)
-				+ SHorizontalBox::Slot()
-				.AutoWidth()
-				.Padding(2.f)
-				.HAlign(HAlign_Left)
-				.VAlign(VAlign_Center)
-				[
-					SNew(SImage)
-					.DesiredSizeOverride(FVector2D(12.f, 12.f))
-					.Image(Icon)
-				]
-				+ SHorizontalBox::Slot()
-				.FillWidth(1.f)
-				.Padding(2.f)
-				.HAlign(HAlign_Left)
-				.VAlign(VAlign_Center)
+				SNew(SCheckBox)
+				.Style(FAppStyle::Get(), "DetailsView.SectionButton")
+				.OnCheckStateChanged(this, &SOperatorStackEditorPanel::OnToolbarButtonClicked, WidgetIdx)
+				.IsChecked(this, &SOperatorStackEditorPanel::IsToolbarButtonActive, WidgetIdx)
+				.Visibility(this, &SOperatorStackEditorPanel::GetToolbarButtonVisibility, WidgetIdx)
 				[
-					SNew(STextBlock)
-					.Text(Label)
+					SNew(SHorizontalBox)
+					+ SHorizontalBox::Slot()
+					.AutoWidth()
+					.Padding(2.f)
+					.HAlign(HAlign_Left)
+					.VAlign(VAlign_Center)
+					[
+						SNew(SImage)
+						.DesiredSizeOverride(FVector2D(12.f, 12.f))
+						.Image(Icon)
+					]
+					+ SHorizontalBox::Slot()
+					.FillWidth(1.f)
+					.Padding(2.f)
+					.HAlign(HAlign_Left)
+					.VAlign(VAlign_Center)
+					[
+						SNew(STextBlock)
+						.Text(Label)
+					]
 				]
 			]
 		];
@@ -279,7 +279,7 @@ void SOperatorStackEditorPanel::UpdateSlots()
 		{
 			UOperatorStackEditorStackCustomization* Customization = EditorSubsystem->GetCustomization(NamedSlot.Key);
 
-			check(Customization);
+			check(Customization)
 
 			CustomizationTrees.Add(Customization, FOperatorStackEditorTree(Customization, Context));
 
diff --git a/Engine/Plugins/Experimental/OperatorStack/Source/OperatorStack/Private/Widgets/SOperatorStackEditorPanel.h b/Engine/Plugins/Experimental/OperatorStack/Source/OperatorStack/Private/Widgets/SOperatorStackEditorPanel.h
--- a/Engine/Plugins/Experimental/OperatorStack/Source/OperatorStack/Private/Widgets/SOperatorStackEditorPanel.h
+++ b/Engine/Plugins/Experimental/OperatorStack/Source/OperatorStack/Private/Widgets/SOperatorStackEditorPanel.h
@@ -50,9 +50,9 @@ public:
 
 	const FOperatorStackEditorTree& GetItemTree(UOperatorStackEditorStackCustomization* InCustomization);
 
-	void SaveItemExpansionState(const void* InItem, bool bInExpanded);
+	void SaveItemExpansionState(uint32 InItem, bool bInExpanded);
 
-	bool GetItemExpansionState(const void* InItem, bool& bOutExpanded);
+	bool GetItemExpansionState(uint32 InItem, bool& bOutExpanded);
 
 	virtual FOperatorStackEditorContextPtr GetContext() const override
 	{
@@ -127,5 +127,5 @@ private:
 	TMap<TWeakObjectPtr<UOperatorStackEditorStackCustomization>, FOperatorStackEditorTree> CustomizationTrees;
 
 	/** To check whether an item was expanded and restore it */
-	TMap<const void*, bool> ItemExpansionState;
+	TMap<uint32, bool> ItemExpansionState;
 };
\ No newline at end of file
diff --git a/Engine/Plugins/Experimental/OperatorStack/Source/OperatorStack/Private/Widgets/SOperatorStackEditorStack.cpp b/Engine/Plugins/Experimental/OperatorStack/Source/OperatorStack/Private/Widgets/SOperatorStackEditorStack.cpp
--- a/Engine/Plugins/Experimental/OperatorStack/Source/OperatorStack/Private/Widgets/SOperatorStackEditorStack.cpp
+++ b/Engine/Plugins/Experimental/OperatorStack/Source/OperatorStack/Private/Widgets/SOperatorStackEditorStack.cpp
@@ -12,6 +12,7 @@
 #include "SOperatorStackEditorPanel.h"
 #include "Items/ICustomDetailsViewItem.h"
 #include "SOperatorStackExpanderButton.h"
+#include "SPositiveActionButton.h"
 #include "Styles/OperatorStackEditorStyle.h"
 #include "Styling/AppStyle.h"
 #include "Styling/ToolBarStyle.h"
@@ -66,9 +67,9 @@ void SOperatorStackEditorStack::Construct(const FArguments& InArgs
 	{
 		Items = ItemTree.GetChildrenItems(InCustomizeItem);
 	}
-	else
+	else if (const FOperatorStackEditorItemPtr RootItem = ItemTree.GetRootItem())
 	{
-		Items = ItemTree.GetRootItems();
+		Items = {RootItem};
 	}
 
 	ChildSlot
@@ -137,7 +138,7 @@ TSharedPtr<SWidget> SOperatorStackEditorStack::GenerateHeaderWidget()
 	{
 		// Find previous expansion state or use default
 		bHeaderExpanded = HeaderBuilder.GetStartsExpanded();
-		MainPanel->GetItemExpansionState(CustomizeItem->GetValuePtr(), bHeaderExpanded);
+		MainPanel->GetItemExpansionState(CustomizeItem->GetHash(), bHeaderExpanded);
 
 		HorizontalHeaderWidget->AddSlot()
 		.AutoWidth()
@@ -302,13 +303,26 @@ TSharedPtr<SWidget> SOperatorStackEditorStack::GenerateHeaderWidget()
 			// Add scrollbox with pinned keywords
 			if (!HeaderBuilder.GetSearchPinnedKeywords().IsEmpty())
 			{
-				static const FToolBarStyle& ToolBarStyle = FAppStyle::Get().GetWidgetStyle<FToolBarStyle>("SlimToolBar");
-				static const FCheckBoxStyle* const CheckStyle = &ToolBarStyle.ToggleButton;
-
 				const TSharedPtr<SScrollBox> SearchScroll = SNew(SScrollBox)
 				.Orientation(EOrientation::Orient_Horizontal)
 				.ScrollBarThickness(FVector2D(3.f));
 
+				auto CreatePinnedButton = [this](const FString& InPinnedKeyword)->TSharedRef<SWidget>
+				{
+					return SNew(SBox)
+						.Padding(FMargin(0.0f))
+						[
+							SNew(SCheckBox)
+							.Style(FAppStyle::Get(), "DetailsView.SectionButton")
+							.OnCheckStateChanged(this, &SOperatorStackEditorStack::OnSearchPinnedKeyword, InPinnedKeyword)
+							[
+								SNew(STextBlock)
+								.TextStyle(FAppStyle::Get(), "SmallText")
+								.Text(FText::FromString(InPinnedKeyword))
+							]
+						];
+				};
+
 				for (const FString& PinnedKeyword : HeaderBuilder.GetSearchPinnedKeywords())
 				{
 					SearchScroll->AddSlot()
@@ -316,14 +330,7 @@ TSharedPtr<SWidget> SOperatorStackEditorStack::GenerateHeaderWidget()
 					.VAlign(VAlign_Fill)
 					.Padding(Padding)
 					[
-						SNew(SCheckBox)
-						.Style(CheckStyle)
-						.ForegroundColor(FLinearColor::White)
-						.OnCheckStateChanged(this, &SOperatorStackEditorStack::OnSearchPinnedKeyword, PinnedKeyword)
-						[
-							SNew(STextBlock)
-							.Text(FText::FromString(PinnedKeyword))
-						]
+						CreatePinnedButton(PinnedKeyword)
 					];
 				}
 
@@ -378,7 +385,11 @@ TSharedPtr<SWidget> SOperatorStackEditorStack::GenerateHeaderWidget()
 		.Padding(Padding)
 		.AutoHeight()
 		[
-			MenuButtonWidget
+			SNew(SPositiveActionButton)
+			.Icon(HeaderBuilder.GetToolMenuIcon())
+			.Text(HeaderBuilder.GetToolMenuLabel())
+			.ToolTipText(HeaderBuilder.GetToolMenuLabel())
+			.OnGetMenuContent(this, &SOperatorStackEditorStack::GenerateMenuWidget, ToolMenuName)
 		];
 	}
 
@@ -456,6 +467,8 @@ TSharedPtr<SWidget> SOperatorStackEditorStack::GenerateBodyWidget()
 	// We are the root and we have multiple supported items selected
 	if (!CustomizeItem.IsValid() && !Items.IsEmpty())
 	{
+		check(Items.Num() == 1)
+
 		FOperatorStackEditorBodyBuilder StackBodyBuilder;
 		StackCustomization->CustomizeStackBody(ItemTree, StackBodyBuilder);
 
@@ -471,35 +484,24 @@ TSharedPtr<SWidget> SOperatorStackEditorStack::GenerateBodyWidget()
 			return StackWidget;
 		}
 
-		const TSharedRef<SSplitter> MultiStackBox = SNew(SSplitter)
-			.Orientation(Orient_Vertical);
-
-		for (FOperatorStackEditorItemPtr Item : Items)
-		{
-			TSharedPtr<SOperatorStackEditorStack> ItemStackWidget;
-
-			MultiStackBox->AddSlot()
-			.Value(1.f)
+		// Scrollbar to be able to scroll in the full view per vertical stack
+		TSharedPtr<SOperatorStackEditorStack> RootStack;
+		TSharedRef<SScrollBox> ScrollBox = SNew(SScrollBox)
+			.ScrollBarPadding(FMargin(0.f))
+			.ScrollBarAlwaysVisible(true)
+			.ScrollBarThickness(FVector2D(3.f))
+			.ScrollBarVisibility(EVisibility::Hidden)
+			.Orientation(Orient_Vertical)
+			+ SScrollBox::Slot()
+			.FillSize(1.f)
+			.Padding(Padding)
 			[
-				// Scrollbar to be able to scroll in the full view per vertical stack
-				SNew(SScrollBox)
-				.ScrollBarPadding(FMargin(0.f))
-				.ScrollBarAlwaysVisible(true)
-				.ScrollBarThickness(FVector2D(3.f))
-				.ScrollBarVisibility(EVisibility::Hidden)
-				.Orientation(Orient_Vertical)
-				+ SScrollBox::Slot()
-				.FillSize(1.f)
-				.Padding(Padding)
-				[
-					SAssignNew(ItemStackWidget, SOperatorStackEditorStack, GetMainPanel(), StackCustomization, Item)
-				]
+				SAssignNew(RootStack, SOperatorStackEditorStack, GetMainPanel(), StackCustomization, Items[0])
 			];
 
-			ItemsWidgets.Add(ItemStackWidget);
-		}
+		ItemsWidgets.Add(RootStack);
 
-		return MultiStackBox;
+		return ScrollBox;
 	}
 
 	// We are not the root but we contain children then add a list view
@@ -827,29 +829,29 @@ bool SOperatorStackEditorStack::HandleSearch(const TSet<FString>& InSearchedKeyw
 		}
 	}
 
-	// Search for properties in details view
+	// Search for properties row in details view
+	bool bRowFound = false;
+
 	if (BodyDetailsView)
 	{
-		bMatchSearch |= BodyDetailsView->FilterItems(InSearchedKeywordsAND.Array());
+		bRowFound |= BodyDetailsView->FilterItems(InSearchedKeywordsAND.Array());
 	}
 
 	if (FooterDetailsView)
 	{
-		bMatchSearch |= FooterDetailsView->FilterItems(InSearchedKeywordsAND.Array());
+		bRowFound |= FooterDetailsView->FilterItems(InSearchedKeywordsAND.Array());
+	}
+
+	if (!InSearchedKeywordsAND.IsEmpty())
+	{
+		bMatchSearch |= bRowFound;
 	}
 
 	// Do not hide root item if nothing was found
 	if (CustomizeItem.IsValid())
 	{
-		if (bMatchSearch || MatchSearch(InSearchedKeywordsOR, InSearchedKeywordsAND))
-		{
-			bMatchSearch = true;
-			SetVisibility(EVisibility::Visible);
-		}
-		else
-		{
-			SetVisibility(EVisibility::Collapsed);
-		}
+		bMatchSearch |= MatchSearch(InSearchedKeywordsOR, InSearchedKeywordsAND);
+		SetVisibility(bMatchSearch ? EVisibility::Visible : EVisibility::Collapsed);
 	}
 
 	return bMatchSearch;
@@ -946,7 +948,7 @@ void SOperatorStackEditorStack::OnHeaderExpansionChanged(bool bInExpansion)
 	{
 		if (CustomizeItem.IsValid())
 		{
-			MainPanel->SaveItemExpansionState(CustomizeItem->GetValuePtr(), bHeaderExpanded);
+			MainPanel->SaveItemExpansionState(CustomizeItem->GetHash(), bHeaderExpanded);
 		}
 	}
 }
@@ -1015,11 +1017,12 @@ FReply SOperatorStackEditorStack::OnHeaderMouseButtonDown(const FGeometry& MyGeo
 
 bool SOperatorStackEditorStack::IsSelectableRow(FOperatorStackEditorItemPtr InItem) const
 {
-	UOperatorStackEditorStackCustomization* StackCustomization = GetStackCustomization();
-
-	if (InItem.IsValid() && StackCustomization)
+	if (InItem.IsValid())
 	{
-		return StackCustomization->OnIsItemDraggable(InItem);
+		if (UOperatorStackEditorStackCustomization* StackCustomization = GetStackCustomization())
+		{
+			return StackCustomization->OnIsItemSelectable(InItem);
+		}
 	}
 
 	return false;
@@ -1031,11 +1034,22 @@ TSharedRef<ICustomDetailsView> SOperatorStackEditorStack::CreateDetailsView(cons
 
 	if (InItem.GetValueType().GetTypeEnum() == EOperatorStackEditorItemType::Object)
 	{
-		UObject* Item = InItem.Get<UObject>();
-		CustomDetailsView->SetObject(Item);
+		if (InItem.GetValueCount() == 1)
+		{
+			UObject* Object = InItem.Get<UObject>(0);
+			CustomDetailsView->SetObject(Object);
+		}
+		else
+		{
+			const TArray<UObject*> Objects = InItem.GetAsArray<UObject>();
+			CustomDetailsView->SetObjects(Objects);
+		}
 	}
 	else if (InItem.GetValueType().GetTypeEnum() == EOperatorStackEditorItemType::Struct)
 	{
+		// multi support for struct on scope is not yet available
+		check(InItem.GetValueCount() == 1)
+
 		const FOperatorStackEditorStructItem* Item = static_cast<const FOperatorStackEditorStructItem*>(&InItem);
 		CustomDetailsView->SetStruct(Item->GetStructOnScope());
 	}
diff --git a/Engine/Plugins/Experimental/OperatorStack/Source/OperatorStack/Private/Widgets/SOperatorStackEditorStack.h b/Engine/Plugins/Experimental/OperatorStack/Source/OperatorStack/Private/Widgets/SOperatorStackEditorStack.h
--- a/Engine/Plugins/Experimental/OperatorStack/Source/OperatorStack/Private/Widgets/SOperatorStackEditorStack.h
+++ b/Engine/Plugins/Experimental/OperatorStack/Source/OperatorStack/Private/Widgets/SOperatorStackEditorStack.h
@@ -26,7 +26,7 @@ class SOperatorStackEditorStack : public SCompoundWidget
 	friend class SOperatorStackEditorStackRow;
 
 public:
-	static inline constexpr float Padding = 3.f;
+	static constexpr float Padding = 2.f;
 	static const TMap<EOperatorStackEditorMessageType, FLinearColor> MessageBoxColors;
 	static const TMap<EOperatorStackEditorMessageType, const FSlateBrush*> MessageBoxIcons;
 
@@ -65,7 +65,7 @@ protected:
 	/** List selection changed */
 	void OnSelectionChanged(FOperatorStackEditorItemPtr InItem, ESelectInfo::Type InSelect) const;
 
-	/** Is this item selectable in the list view, here check if item is draggable */
+	/** Is this item selectable in the list view */
 	bool IsSelectableRow(FOperatorStackEditorItemPtr InItem) const;
 
 	/** A stack widget consists of a header, body and footer */
@@ -86,6 +86,9 @@ protected:
 	const FSlateBrush* GetMessageBoxIcon() const;
 	EVisibility GetMessageBoxIconVisibility() const;
 
+	/** The customization to use */
+	TWeakObjectPtr<UOperatorStackEditorStackCustomization> StackCustomizationWeak;
+
 	/** The main panel we are in */
 	TWeakPtr<SOperatorStackEditorPanel> MainPanelWeak;
 
@@ -107,9 +110,6 @@ protected:
 	/** Footer custom view for properties in footer */
 	TSharedPtr<ICustomDetailsView> FooterDetailsView;
 
-	/** The customization to use */
-	TWeakObjectPtr<UOperatorStackEditorStackCustomization> StackCustomizationWeak;
-
 	/** Children items widget created from this, used for search propagation */
 	TArray<TSharedPtr<SOperatorStackEditorStack>> ItemsWidgets;
 
diff --git a/Engine/Plugins/Experimental/OperatorStack/Source/OperatorStack/Private/Widgets/TableRows/SOperatorStackEditorStackRow.cpp b/Engine/Plugins/Experimental/OperatorStack/Source/OperatorStack/Private/Widgets/TableRows/SOperatorStackEditorStackRow.cpp
--- a/Engine/Plugins/Experimental/OperatorStack/Source/OperatorStack/Private/Widgets/TableRows/SOperatorStackEditorStackRow.cpp
+++ b/Engine/Plugins/Experimental/OperatorStack/Source/OperatorStack/Private/Widgets/TableRows/SOperatorStackEditorStackRow.cpp
@@ -72,14 +72,22 @@ FReply SOperatorStackEditorStackRow::OnStackDragDetected(const FGeometry& InGeom
 {
 	const TSharedPtr<SOperatorStackEditorStack> OuterStack = OuterStackWeak.Pin();
 
-	if (OuterStack.IsValid() && OuterStack->ItemsListView.IsValid())
+	if (OuterStack.IsValid() && OuterStack->ItemsListView.IsValid() && OuterStack->ItemsListView->GetNumItemsSelected() > 0)
 	{
-		const TArray<TSharedPtr<FOperatorStackEditorItem>>& SelectedItems = OuterStack->ItemsListView->GetSelectedItems();
+		UOperatorStackEditorStackCustomization* Customization = OuterStack->GetStackCustomization();
+		TArray<TSharedPtr<FOperatorStackEditorItem>> SelectedItems;
+
+		for (FOperatorStackEditorItemPtr SelectedItem : OuterStack->ItemsListView->GetSelectedItems())
+		{
+			if (SelectedItem.IsValid() && Customization->OnIsItemDraggable(SelectedItem))
+			{
+				SelectedItems.Add(SelectedItem);
+			}
+		}
 
 		if (!SelectedItems.IsEmpty())
 		{
 			const TSharedRef<FOperatorStackEditorDragDropOp> StackDragDropOp = FOperatorStackEditorDragDropOp::New(SelectedItems);
-
 			return FReply::Handled().BeginDragDrop(StackDragDropOp);
 		}
 	}
diff --git a/Engine/Plugins/Experimental/OperatorStack/Source/OperatorStack/Public/Customizations/OperatorStackEditorStackCustomization.h b/Engine/Plugins/Experimental/OperatorStack/Source/OperatorStack/Public/Customizations/OperatorStackEditorStackCustomization.h
--- a/Engine/Plugins/Experimental/OperatorStack/Source/OperatorStack/Public/Customizations/OperatorStackEditorStackCustomization.h
+++ b/Engine/Plugins/Experimental/OperatorStack/Source/OperatorStack/Public/Customizations/OperatorStackEditorStackCustomization.h
@@ -52,11 +52,17 @@ public:
 	/** Checks if this customization is supported for this item */
 	bool IsCustomizationSupportedFor(const FOperatorStackEditorItemPtr& InItem) const;
 
-	/**
-	 * Transform a context item in items that this customization supports or returns false when not supported
-	 * For example, context item is an actor but this customization supports a specific component then look into actor for this component here
-	 */
-	virtual bool TransformContextItem(const FOperatorStackEditorItemPtr& InItem, TArray<FOperatorStackEditorItemPtr>& OutTransformedItems) const { return false; }
+	/** Retrieves root supported item out of the provided context, Children classes must implement this */
+	virtual bool GetRootItem(const FOperatorStackEditorContext& InContext, FOperatorStackEditorItemPtr& OutRootItem) const
+	{
+		return OutRootItem.IsValid() && OutRootItem->HasValue();
+	}
+
+	/** Retrieves the supported children items based out of a parent item, Children classes must implement this */
+	virtual bool GetChildrenItem(const FOperatorStackEditorItemPtr& InItem, TArray<FOperatorStackEditorItemPtr>& OutChildrenItems) const
+	{
+		return !OutChildrenItems.IsEmpty();
+	}
 
 	/** Customize the header for a context */
 	virtual void CustomizeStackHeader(const FOperatorStackEditorTree& InItemTree, FOperatorStackEditorHeaderBuilder& InStackHeaderBuilder) {}
@@ -73,6 +79,12 @@ public:
 	/** Customize the body for a supported item */
 	virtual void CustomizeItemFooter(const FOperatorStackEditorItemPtr& InItem, const FOperatorStackEditorTree& InItemTree, FOperatorStackEditorFooterBuilder& InFooterBuilder) {}
 
+	/** Can an item be selected */
+	virtual bool OnIsItemSelectable(const FOperatorStackEditorItemPtr& InItem)
+	{
+		return true;
+	}
+
 	/** Can an item be dragged */
 	virtual bool OnIsItemDraggable(const FOperatorStackEditorItemPtr& InItem)
 	{
diff --git a/Engine/Plugins/Experimental/OperatorStack/Source/OperatorStack/Public/Items/OperatorStackEditorGroupItem.h b/Engine/Plugins/Experimental/OperatorStack/Source/OperatorStack/Public/Items/OperatorStackEditorGroupItem.h
new file mode 100644
--- /dev/null
+++ b/Engine/Plugins/Experimental/OperatorStack/Source/OperatorStack/Public/Items/OperatorStackEditorGroupItem.h
@@ -0,0 +1,69 @@
+// Copyright Epic Games, Inc. All Rights Reserved.
+
+#pragma once
+
+#include "OperatorStackEditorItem.h"
+
+/** Groups item of the same type together */
+struct FOperatorStackEditorGroupItem : FOperatorStackEditorItem
+{
+	explicit FOperatorStackEditorGroupItem(const TArray<FOperatorStackEditorItemPtr>& InItems, const FOperatorStackEditorItemType& InType)
+		: FOperatorStackEditorItem(
+			InType
+		)
+	{
+		CachedHash = 0;
+
+		for (const FOperatorStackEditorItemPtr& Item : InItems)
+		{
+			if (Item.IsValid())
+			{
+				// No groups into groups
+				check(Item->GetValueCount() == 1)
+
+				// Group is only allowed for item of the same type
+				check(GetValueType() == Item->GetValueType())
+
+				CachedHash += Item->GetHash();
+				Items.Add(Item);
+			}
+		}
+	}
+
+	virtual uint32 GetValueCount() const override
+	{
+		return Items.Num();
+	}
+
+	virtual bool HasValue() const override
+	{
+		for (const FOperatorStackEditorItemPtr& Item : Items)
+		{
+			if (Item.IsValid() && Item->HasValue())
+			{
+				return true;
+			}
+		}
+
+		return FOperatorStackEditorItem::HasValue();
+	}
+
+	virtual uint32 GetHash() const override
+	{
+		return CachedHash;
+	}
+
+	virtual void* GetValuePtr(int32 InIndex) const override
+	{
+		if (Items.IsValidIndex(InIndex) && Items[InIndex].IsValid())
+		{
+			// No groups allowed so pick first item
+			return Items[InIndex]->GetValuePtr(0);
+		}
+
+		return FOperatorStackEditorItem::GetValuePtr(InIndex);
+	}
+
+protected:
+	TArray<FOperatorStackEditorItemPtr> Items;
+};
\ No newline at end of file
diff --git a/Engine/Plugins/Experimental/OperatorStack/Source/OperatorStack/Public/Items/OperatorStackEditorItem.h b/Engine/Plugins/Experimental/OperatorStack/Source/OperatorStack/Public/Items/OperatorStackEditorItem.h
--- a/Engine/Plugins/Experimental/OperatorStack/Source/OperatorStack/Public/Items/OperatorStackEditorItem.h
+++ b/Engine/Plugins/Experimental/OperatorStack/Source/OperatorStack/Public/Items/OperatorStackEditorItem.h
@@ -15,8 +15,6 @@
 /** Abstract base parent item */
 struct FOperatorStackEditorItem
 {
-	friend class SOperatorStackEditorStack;
-
 	explicit FOperatorStackEditorItem(const FOperatorStackEditorItemType& InType)
 		: ItemType(InType)
 	{}
@@ -73,9 +71,11 @@ struct FOperatorStackEditorItem
 	 * You can use this to get the item Item.Get<UObject>() or Item.Get<FMyStruct>() or Item.Get<int32>()
 	 */
 	template <typename InValueType>
-	InValueType* Get() const
+	InValueType* Get(uint32 InIndex) const
 	{
-		if (!GetValuePtr())
+		check(InIndex < GetValueCount())
+
+		if (!GetValuePtr(InIndex))
 		{
 			return nullptr;
 		}
@@ -84,39 +84,76 @@ struct FOperatorStackEditorItem
 		{
 			if (ItemType.GetTypeEnum() == EOperatorStackEditorItemType::Object)
 			{
-				return Cast<InValueType>(static_cast<UObject*>(GetValuePtr()));
+				return Cast<InValueType>(static_cast<UObject*>(GetValuePtr(InIndex)));
 			}
 		}
 		else if constexpr (TModels_V<CStaticStructProvider, InValueType> || TModels_V<CBaseStructureProvider, InValueType>)
 		{
 			if (ItemType.GetTypeEnum() == EOperatorStackEditorItemType::Struct)
 			{
-				return static_cast<InValueType*>(GetValuePtr());
+				return static_cast<InValueType*>(GetValuePtr(InIndex));
 			}
 		}
 		else if constexpr (TIsPODType<InValueType>::Value)
 		{
 			if (ItemType.GetTypeEnum() == EOperatorStackEditorItemType::Primitive)
 			{
-				return static_cast<InValueType*>(GetValuePtr());
+				return static_cast<InValueType*>(GetValuePtr(InIndex));
 			}
 		}
 
 		return nullptr;
 	}
 
+	/**
+	 * You can use this to get all the underlying items as an array of a specific type
+	 */
+	template <typename InValueType>
+	TArray<InValueType*> GetAsArray() const
+	{
+		uint32 ValueCount = GetValueCount();
+
+		TArray<InValueType*> Values;
+		Values.Reserve(ValueCount);
+
+		for (uint32 Index = 0; Index < ValueCount; Index++)
+		{
+			Values.Add(Get<InValueType>(Index));
+		}
+
+		return Values;
+	}
+
 	/** Get the value type of this item */
 	const FOperatorStackEditorItemType& GetValueType() const
 	{
 		return ItemType;
 	}
 
+	/** Get the amount of value stored within this item */
+	virtual uint32 GetValueCount() const
+	{
+		return 0;
+	}
+
 	/** Checks if this item has a value and that it is usable */
 	virtual bool HasValue() const
 	{
 		return false;
 	}
 
+	/** Get raw ptr to value, prefer using Get<>() instead */
+	virtual void* GetValuePtr(int32 InIndex) const
+	{
+		return nullptr;
+	}
+
+	/** Override in child to be able to compare item */
+	virtual uint32 GetHash() const
+	{
+		return CachedHash;
+	}
+
 	friend uint32 GetTypeHash(const FOperatorStackEditorItem& InItem)
 	{
 		return InItem.GetHash();
@@ -133,18 +170,7 @@ struct FOperatorStackEditorItem
 	}
 
 protected:
-	/** Override in child to be able to compare item */
-	virtual uint32 GetHash() const
-	{
-		return 0;
-	}
-
-	/** Get raw ptr to value, prefer using Get<>() instead */
-	virtual void* GetValuePtr() const
-	{
-		return nullptr;
-	}
-
+	uint32 CachedHash = 0;
 	FOperatorStackEditorItemType ItemType;
 };
 
diff --git a/Engine/Plugins/Experimental/OperatorStack/Source/OperatorStack/Public/Items/OperatorStackEditorObjectItem.h b/Engine/Plugins/Experimental/OperatorStack/Source/OperatorStack/Public/Items/OperatorStackEditorObjectItem.h
--- a/Engine/Plugins/Experimental/OperatorStack/Source/OperatorStack/Public/Items/OperatorStackEditorObjectItem.h
+++ b/Engine/Plugins/Experimental/OperatorStack/Source/OperatorStack/Public/Items/OperatorStackEditorObjectItem.h
@@ -25,22 +25,26 @@ struct FOperatorStackEditorObjectItem : FOperatorStackEditorItem
 		CachedHash = GetTypeHash(InItem);
 	}
 
+	virtual uint32 GetValueCount() const override
+	{
+		return 1;
+	}
+
 	virtual bool HasValue() const override
 	{
 		return ObjectWeak.IsValid();
 	}
 
-protected:
 	virtual uint32 GetHash() const override
 	{
 		return CachedHash;
 	}
 
-	virtual void* GetValuePtr() const override
+	virtual void* GetValuePtr(int32 InIndex) const override
 	{
 		return ObjectWeak.Get();
 	}
 
-	uint32 CachedHash = 0;
+protected:
 	TWeakObjectPtr<UObject> ObjectWeak;
 };
\ No newline at end of file
diff --git a/Engine/Plugins/Experimental/OperatorStack/Source/OperatorStack/Public/Items/OperatorStackEditorPrimitiveItem.h b/Engine/Plugins/Experimental/OperatorStack/Source/OperatorStack/Public/Items/OperatorStackEditorPrimitiveItem.h
--- a/Engine/Plugins/Experimental/OperatorStack/Source/OperatorStack/Public/Items/OperatorStackEditorPrimitiveItem.h
+++ b/Engine/Plugins/Experimental/OperatorStack/Source/OperatorStack/Public/Items/OperatorStackEditorPrimitiveItem.h
@@ -25,22 +25,26 @@ struct FOperatorStackEditorPrimitiveItem : FOperatorStackEditorItem
 		CachedHash = GetTypeHash(InValue);
 	}
 
+	virtual uint32 GetValueCount() const override
+	{
+		return 1;
+	}
+
 	virtual bool HasValue() const override
 	{
 		return MemoryPtr != nullptr;
 	}
 
-protected:
 	virtual uint32 GetHash() const override
 	{
 		return CachedHash;
 	}
 
-	virtual void* GetValuePtr() const override
+	virtual void* GetValuePtr(int32 InIndex) const override
 	{
 		return MemoryPtr;
 	}
 
-	uint32 CachedHash = 0;
+protected:
 	uint8* MemoryPtr = nullptr;
 };
\ No newline at end of file
diff --git a/Engine/Plugins/Experimental/OperatorStack/Source/OperatorStack/Public/Items/OperatorStackEditorStructItem.h b/Engine/Plugins/Experimental/OperatorStack/Source/OperatorStack/Public/Items/OperatorStackEditorStructItem.h
--- a/Engine/Plugins/Experimental/OperatorStack/Source/OperatorStack/Public/Items/OperatorStackEditorStructItem.h
+++ b/Engine/Plugins/Experimental/OperatorStack/Source/OperatorStack/Public/Items/OperatorStackEditorStructItem.h
@@ -65,6 +65,11 @@ struct FOperatorStackEditorStructItem : FOperatorStackEditorItem
 		return StructWeak.Pin();
 	}
 
+	virtual uint32 GetValueCount() const override
+	{
+		return 1;
+	}
+
 	virtual bool HasValue() const override
 	{
 		if (const TSharedPtr<FStructOnScope> Struct = StructWeak.Pin())
@@ -75,13 +80,12 @@ struct FOperatorStackEditorStructItem : FOperatorStackEditorItem
 		return false;
 	}
 
-protected:
 	virtual uint32 GetHash() const override
 	{
 		return CachedHash;
 	}
 
-	virtual void* GetValuePtr() const override
+	virtual void* GetValuePtr(int32 InIndex) const override
 	{
 		if (const TSharedPtr<FStructOnScope> Struct = StructWeak.Pin())
 		{
@@ -91,6 +95,6 @@ protected:
 		return nullptr;
 	}
 
-	uint32 CachedHash = 0;
+protected:
 	TWeakPtr<FStructOnScope> StructWeak;
 };
\ No newline at end of file
diff --git a/Engine/Plugins/Experimental/OperatorStack/Source/OperatorStack/Public/Items/OperatorStackEditorTree.h b/Engine/Plugins/Experimental/OperatorStack/Source/OperatorStack/Public/Items/OperatorStackEditorTree.h
--- a/Engine/Plugins/Experimental/OperatorStack/Source/OperatorStack/Public/Items/OperatorStackEditorTree.h
+++ b/Engine/Plugins/Experimental/OperatorStack/Source/OperatorStack/Public/Items/OperatorStackEditorTree.h
@@ -13,7 +13,7 @@ struct FOperatorStackEditorTree
 	explicit FOperatorStackEditorTree(UOperatorStackEditorStackCustomization* InCustomization, FOperatorStackEditorContextPtr InContext);
 
 	/** Get all top items supported by this customization */
-	OPERATORSTACKEDITOR_API TArray<FOperatorStackEditorItemPtr> GetRootItems() const;
+	OPERATORSTACKEDITOR_API FOperatorStackEditorItemPtr GetRootItem() const;
 
 	/** Get supported children items from a supported item */
 	TArray<FOperatorStackEditorItemPtr> GetChildrenItems(FOperatorStackEditorItemPtr InItem) const;
@@ -49,10 +49,7 @@ private:
 
 	using FOperatorStackEditorTreeNodePtr = TSharedPtr<FOperatorStackEditorTreeNode>;
 
-	void GetSupportedItems(const TArray<FOperatorStackEditorItemPtr>& InItems, TArray<FOperatorStackEditorItemPtr>& OutSupportedItems) const;
 	TArray<FOperatorStackEditorItemPtr> GetSupportedChildrenItems(const FOperatorStackEditorItemPtr& InParentItem) const;
-
-	void BuildTree(const TArray<FOperatorStackEditorItemPtr>& InSourceItems);
 	void BuildTreeInternal(const TArray<FOperatorStackEditorItemPtr>& InItems, int32 InParent);
 
 	/** All items in this tree */
@@ -62,7 +59,7 @@ private:
 	TArray<FOperatorStackEditorTreeNodePtr> Nodes;
 
 	/** Root items in the tree */
-	TArray<FOperatorStackEditorTreeNodePtr> RootNodes;
+	FOperatorStackEditorTreeNodePtr RootNode;
 
 	/** Customization used to build this tree */
 	TWeakObjectPtr<UOperatorStackEditorStackCustomization> CustomizationWeak;
diff --git a/Engine/Plugins/Experimental/PropertyAnimatorCore/Source/PropertyAnimatorCore/Public/Animators/PropertyAnimatorCoreBase.h b/Engine/Plugins/Experimental/PropertyAnimatorCore/Source/PropertyAnimatorCore/Public/Animators/PropertyAnimatorCoreBase.h
--- a/Engine/Plugins/Experimental/PropertyAnimatorCore/Source/PropertyAnimatorCore/Public/Animators/PropertyAnimatorCoreBase.h
+++ b/Engine/Plugins/Experimental/PropertyAnimatorCore/Source/PropertyAnimatorCore/Public/Animators/PropertyAnimatorCoreBase.h
@@ -56,7 +56,7 @@ public:
 
 	PROPERTYANIMATORCORE_API AActor* GetAnimatorActor() const;
 
-	UPropertyAnimatorCoreComponent* GetAnimatorComponent() const;
+	PROPERTYANIMATORCORE_API UPropertyAnimatorCoreComponent* GetAnimatorComponent() const;
 
 	/** Get the global animators magnitude */
 	float GetAnimatorComponentMagnitude() const;
diff --git a/Engine/Plugins/Experimental/PropertyAnimatorCore/Source/PropertyAnimatorCoreEditor/Private/Customizations/PropertyAnimatorCoreEditorContextTypeCustomization.cpp b/Engine/Plugins/Experimental/PropertyAnimatorCore/Source/PropertyAnimatorCoreEditor/Private/Customizations/PropertyAnimatorCoreEditorContextTypeCustomization.cpp
--- a/Engine/Plugins/Experimental/PropertyAnimatorCore/Source/PropertyAnimatorCoreEditor/Private/Customizations/PropertyAnimatorCoreEditorContextTypeCustomization.cpp
+++ b/Engine/Plugins/Experimental/PropertyAnimatorCore/Source/PropertyAnimatorCoreEditor/Private/Customizations/PropertyAnimatorCoreEditorContextTypeCustomization.cpp
@@ -77,14 +77,17 @@ void FPropertyAnimatorCoreEditorContextTypeCustomization::CustomizeHeader(TShare
 
 void FPropertyAnimatorCoreEditorContextTypeCustomization::CustomizeChildren(TSharedRef<IPropertyHandle> InPropertyHandle, IDetailChildrenBuilder& InBuilder, IPropertyTypeCustomizationUtils& InUtils)
 {
+	UObject* PropertyObject;
+	if (FPropertyAccess::Success != InPropertyHandle->GetValue(PropertyObject))
+	{
+		return;
+	}
+
 	static const TSet<FName> SkipProperties
 	{
 		GET_MEMBER_NAME_CHECKED(UPropertyAnimatorCoreContext, bAnimated)
 	};
 
-	UObject* PropertyObject;
-	InPropertyHandle->GetValue(PropertyObject);
-
 	if (UPropertyAnimatorCoreContext* Options = Cast<UPropertyAnimatorCoreContext>(PropertyObject))
 	{
 		for (const FProperty* Property : TFieldRange<FProperty>(Options->GetClass()))
diff --git a/Engine/Plugins/Experimental/PropertyAnimatorCore/Source/PropertyAnimatorCoreEditor/Private/Customizations/PropertyAnimatorCoreEditorStackCustomization.cpp b/Engine/Plugins/Experimental/PropertyAnimatorCore/Source/PropertyAnimatorCoreEditor/Private/Customizations/PropertyAnimatorCoreEditorStackCustomization.cpp
--- a/Engine/Plugins/Experimental/PropertyAnimatorCore/Source/PropertyAnimatorCoreEditor/Private/Customizations/PropertyAnimatorCoreEditorStackCustomization.cpp
+++ b/Engine/Plugins/Experimental/PropertyAnimatorCore/Source/PropertyAnimatorCoreEditor/Private/Customizations/PropertyAnimatorCoreEditorStackCustomization.cpp
@@ -2,10 +2,11 @@
 
 #include "PropertyAnimatorCoreEditorStackCustomization.h"
 
+#include "Animators/PropertyAnimatorCoreBase.h"
 #include "Components/PropertyAnimatorCoreComponent.h"
 #include "Contexts/OperatorStackEditorMenuContext.h"
-#include "Animators/PropertyAnimatorCoreBase.h"
 #include "Framework/Commands/GenericCommands.h"
+#include "Items/OperatorStackEditorGroupItem.h"
 #include "Items/OperatorStackEditorObjectItem.h"
 #include "Menus/PropertyAnimatorCoreEditorMenu.h"
 #include "Styling/SlateIconFinder.h"
@@ -17,6 +18,8 @@
 
 #define LOCTEXT_NAMESPACE "PropertyAnimatorEditorStackCustomization"
 
+DEFINE_LOG_CATEGORY_STATIC(LogPropertyAnimatorCoreEditorStackCustomization, Log, All);
+
 UPropertyAnimatorCoreEditorStackCustomization::UPropertyAnimatorCoreEditorStackCustomization()
 	: UOperatorStackEditorStackCustomization(
 		TEXT("Animators")
@@ -28,37 +31,92 @@ UPropertyAnimatorCoreEditorStackCustomization::UPropertyAnimatorCoreEditorStackC
 	RegisterCustomizationFor(UPropertyAnimatorCoreComponent::StaticClass());
 }
 
-bool UPropertyAnimatorCoreEditorStackCustomization::TransformContextItem(const FOperatorStackEditorItemPtr& InItem, TArray<FOperatorStackEditorItemPtr>& OutTransformedItems) const
+bool UPropertyAnimatorCoreEditorStackCustomization::GetRootItem(const FOperatorStackEditorContext& InContext, FOperatorStackEditorItemPtr& OutRootItem) const
 {
-	// If we have an actor then display controllers
-	if (InItem->IsA<AActor>())
-	{
-		const AActor* Actor = InItem->Get<AActor>();
+	TArray<FOperatorStackEditorItemPtr> RootItems;
 
-		if (UActorComponent* ControllerComponent = Actor->FindComponentByClass(UPropertyAnimatorCoreComponent::StaticClass()))
+	// Pick all property animator component as root for the stack view
+	for (const FOperatorStackEditorItemPtr& Item : InContext.GetItems())
+	{
+		if (!Item.IsValid())
 		{
-			OutTransformedItems.Add(MakeShared<FOperatorStackEditorObjectItem>(ControllerComponent));
+			continue;
 		}
 
-		return true;
+		if (Item->IsA<AActor>())
+		{
+			for (const AActor* Actor : Item->GetAsArray<AActor>())
+			{
+				if (UActorComponent* AnimatorComponent = Actor->FindComponentByClass(UPropertyAnimatorCoreComponent::StaticClass()))
+				{
+					RootItems.Add(MakeShared<FOperatorStackEditorObjectItem>(AnimatorComponent));
+				}
+			}
+		}
+		else if (Item->IsA<UPropertyAnimatorCoreComponent>())
+		{
+			for (UPropertyAnimatorCoreComponent* Component : Item->GetAsArray<UPropertyAnimatorCoreComponent>())
+			{
+				RootItems.Add(MakeShared<FOperatorStackEditorObjectItem>(Component));
+			}
+		}
+		else if (Item->IsA<UPropertyAnimatorCoreBase>())
+		{
+			for (const UPropertyAnimatorCoreBase* Animator : Item->GetAsArray<UPropertyAnimatorCoreBase>())
+			{
+				RootItems.Add(MakeShared<FOperatorStackEditorObjectItem>(Animator->GetAnimatorComponent()));
+			}
+		}
 	}
-	// If we have a property controller component then display stack
-	else if (InItem->IsA<UPropertyAnimatorCoreComponent>())
+
+	OutRootItem = MakeShared<FOperatorStackEditorGroupItem>(RootItems, FOperatorStackEditorItemType(UPropertyAnimatorCoreComponent::StaticClass(), EOperatorStackEditorItemType::Object));
+
+	return Super::GetRootItem(InContext, OutRootItem);
+}
+
+bool UPropertyAnimatorCoreEditorStackCustomization::GetChildrenItem(const FOperatorStackEditorItemPtr& InItem, TArray<FOperatorStackEditorItemPtr>& OutChildrenItems) const
+{
+	if (InItem->IsA<UPropertyAnimatorCoreComponent>())
 	{
-		const UPropertyAnimatorCoreComponent* ControllerComponent = InItem->Get<UPropertyAnimatorCoreComponent>();
+		if (InItem->GetValueCount() > 1)
+		{
+			TMap<UClass*, int32> ClassToIndex;
+			TArray<TArray<FOperatorStackEditorItemPtr>> Animators;
+
+			for (const UPropertyAnimatorCoreComponent* AnimatorComponent : InItem->GetAsArray<UPropertyAnimatorCoreComponent>())
+			{
+				for (UPropertyAnimatorCoreBase* Animator : AnimatorComponent->GetAnimators())
+				{
+					if (const int32* Index = ClassToIndex.Find(Animator->GetClass()))
+					{
+						Animators[*Index].Add(MakeShared<FOperatorStackEditorObjectItem>(Animator));
+					}
+					else
+					{
+						TArray<FOperatorStackEditorItemPtr> ModifierGroup;
+						ModifierGroup.Add(MakeShared<FOperatorStackEditorObjectItem>(Animator));
+						int32 GroupIndex = Animators.Add(ModifierGroup);
+						ClassToIndex.Add(Animator->GetClass(), GroupIndex);
+					}
+				}
+			}
 
-		for (UPropertyAnimatorCoreBase* Controller : ControllerComponent->GetAnimators())
+			for (int32 Index = 0; Index < Animators.Num(); Index++)
+			{
+				const UClass* const* AnimatorClass = ClassToIndex.FindKey(Index);
+				OutChildrenItems.Add(MakeShared<FOperatorStackEditorGroupItem>(Animators[Index], FOperatorStackEditorItemType(*AnimatorClass, EOperatorStackEditorItemType::Object)));
+			}
+		}
+		else
 		{
-			if (Controller && Controller->IsA<UPropertyAnimatorCoreBase>())
+			for (UPropertyAnimatorCoreBase* Animator : InItem->Get<UPropertyAnimatorCoreComponent>(0)->GetAnimators())
 			{
-				OutTransformedItems.Add(MakeShared<FOperatorStackEditorObjectItem>(Controller));
+				OutChildrenItems.Add(MakeShared<FOperatorStackEditorObjectItem>(Animator));
 			}
 		}
-
-		return true;
 	}
 
-	return Super::TransformContextItem(InItem, OutTransformedItems);
+	return Super::GetChildrenItem(InItem, OutChildrenItems);
 }
 
 void UPropertyAnimatorCoreEditorStackCustomization::CustomizeStackHeader(const FOperatorStackEditorTree& InItemTree, FOperatorStackEditorHeaderBuilder& InHeaderBuilder)
@@ -68,27 +126,31 @@ void UPropertyAnimatorCoreEditorStackCustomization::CustomizeStackHeader(const F
 		static const FName AddAnimatorMenuName = TEXT("AddAnimatorMenu");
 		if (!UToolMenus::Get()->IsMenuRegistered(AddAnimatorMenuName))
 		{
-			UToolMenu* const AddControllerMenu = UToolMenus::Get()->RegisterMenu(AddAnimatorMenuName, NAME_None, EMultiBoxType::Menu);
-			AddControllerMenu->AddDynamicSection(TEXT("FillAddAnimatorMenuSection"), FNewToolMenuDelegate::CreateUObject(this, &UPropertyAnimatorCoreEditorStackCustomization::FillAddAnimatorMenuSection));
+			UToolMenu* const AddAnimatorMenu = UToolMenus::Get()->RegisterMenu(AddAnimatorMenuName, NAME_None, EMultiBoxType::Menu);
+			AddAnimatorMenu->AddDynamicSection(TEXT("FillAddAnimatorMenuSection"), FNewToolMenuDelegate::CreateUObject(this, &UPropertyAnimatorCoreEditorStackCustomization::FillAddAnimatorMenuSection));
 		}
 
 		// Pinned search keywords
-		TSet<FString> PinnedControllerNames;
+		TSet<FString> PinnedAnimatorNames;
 		for (const FOperatorStackEditorItemPtr& SupportedItem : InItemTree.GetAllItems())
 		{
-			if (const UPropertyAnimatorCoreBase* Controller = SupportedItem->Get<UPropertyAnimatorCoreBase>())
+			if (SupportedItem.IsValid() && SupportedItem->IsA<UPropertyAnimatorCoreBase>())
 			{
-				PinnedControllerNames.Add(Controller->GetAnimatorOriginalName().ToString());
+				if (const UPropertyAnimatorCoreBase* Animator = SupportedItem->Get<UPropertyAnimatorCoreBase>(0))
+				{
+					PinnedAnimatorNames.Add(Animator->GetAnimatorOriginalName().ToString());
+				}
 			}
 		}
 
 		InHeaderBuilder
 			.SetToolMenu(
 				AddAnimatorMenuName
-				, LOCTEXT("AddControllersMenu", "Add Animators")
-				, FAppStyle::GetBrush(TEXT("Icons.PlusCircle")))
+				, LOCTEXT("AddAnimatorsMenu", "Add Animators")
+				, FAppStyle::GetBrush("Icons.Plus")
+			)
 			.SetSearchAllowed(true)
-			.SetSearchPinnedKeywords(PinnedControllerNames);
+			.SetSearchPinnedKeywords(PinnedAnimatorNames);
 	}
 
 	Super::CustomizeStackHeader(InItemTree, InHeaderBuilder);
@@ -96,6 +158,22 @@ void UPropertyAnimatorCoreEditorStackCustomization::CustomizeStackHeader(const F
 
 void UPropertyAnimatorCoreEditorStackCustomization::CustomizeItemHeader(const FOperatorStackEditorItemPtr& InItem, const FOperatorStackEditorTree& InItemTree, FOperatorStackEditorHeaderBuilder& InHeaderBuilder)
 {
+	// Action menu available in header in slim toolbar
+	static const FName HeaderAnimatorMenuName = TEXT("HeaderAnimatorMenu");
+	if (!UToolMenus::Get()->IsMenuRegistered(HeaderAnimatorMenuName))
+	{
+		UToolMenu* const HeaderModifierMenu = UToolMenus::Get()->RegisterMenu(HeaderAnimatorMenuName, NAME_None, EMultiBoxType::SlimHorizontalToolBar);
+		HeaderModifierMenu->AddDynamicSection(TEXT("FillHeaderAnimatorMenu"), FNewToolMenuDelegate::CreateUObject(this, &UPropertyAnimatorCoreEditorStackCustomization::FillAnimatorHeaderActionMenu));
+	}
+
+	// Context menu available when right clicking on item
+	static const FName ContextAnimatorMenuName = TEXT("ContextAnimatorMenu");
+	if (!UToolMenus::Get()->IsMenuRegistered(ContextAnimatorMenuName))
+	{
+		UToolMenu* const ContextModifierMenu = UToolMenus::Get()->RegisterMenu(ContextAnimatorMenuName, NAME_None, EMultiBoxType::Menu);
+		ContextModifierMenu->AddDynamicSection(TEXT("FillContextAnimatorMenu"), FNewToolMenuDelegate::CreateUObject(this, &UPropertyAnimatorCoreEditorStackCustomization::FillAnimatorContextActionMenu));
+	}
+
 	// Customize component header
 	if (InItem->IsA<UPropertyAnimatorCoreComponent>())
 	{
@@ -103,58 +181,40 @@ void UPropertyAnimatorCoreEditorStackCustomization::CustomizeItemHeader(const FO
 
 		const FSlateIcon ClassIcon = FSlateIconFinder::FindIconForClass(UPropertyAnimatorCoreComponent::StaticClass());
 
-		// Action menu available in header in slim toolbar
-		static const FName HeaderComponentMenuName = TEXT("HeaderComponentMenu");
-		if (!UToolMenus::Get()->IsMenuRegistered(HeaderComponentMenuName))
-		{
-			UToolMenu* const HeaderAnimatorMenu = UToolMenus::Get()->RegisterMenu(HeaderComponentMenuName, NAME_None, EMultiBoxType::SlimHorizontalToolBar);
-			HeaderAnimatorMenu->AddDynamicSection(TEXT("FillHeaderComponentMenu"), FNewToolMenuDelegate::CreateUObject(this, &UPropertyAnimatorCoreEditorStackCustomization::FillComponentHeaderActionMenu));
-		}
+		// Commands for item on key events
+		const TSharedPtr<FUICommandList> ComponentCommands = CreateAnimatorCommands(InItem);
 
 		FString HeaderLabel = TEXT("Animators");
+		TAttribute<EOperatorStackEditorMessageType> MessageType = EOperatorStackEditorMessageType::None;
+		TAttribute<FText> MessageText = FText::GetEmpty();
 
-		// Add actor name next to label when we customize multiple items
-		if (InItemTree.GetRootItems().Num() > 1)
+		if (InItem->GetValueCount() > 1)
 		{
-			if (const UPropertyAnimatorCoreComponent* ControllerComponent = InItem->Get<UPropertyAnimatorCoreComponent>())
-			{
-				HeaderLabel += + TEXT(" (") + ControllerComponent->GetOwner()->GetActorNameOrLabel() + TEXT(")");
-			}
+			HeaderLabel += TEXT(" (") + FString::FromInt(InItem->GetValueCount()) + TEXT(")");
+			MessageType = EOperatorStackEditorMessageType::Info;
+			MessageText = LOCTEXT("MultiAnimatorView", "You are viewing multiple items");
 		}
 
 		InHeaderBuilder
 			.SetProperty(EnableProperty)
-			.SetToolbarMenu(HeaderComponentMenuName)
 			.SetIcon(ClassIcon.GetIcon())
-			.SetLabel(FText::FromString(HeaderLabel));
+			.SetLabel(FText::FromString(HeaderLabel))
+			.SetToolbarMenu(HeaderAnimatorMenuName)
+			.SetContextMenu(ContextAnimatorMenuName)
+			.SetCommandList(ComponentCommands)
+			.SetMessageBox(MessageType, MessageText);
 	}
-	// Customize controller header
+	// Customize animator header
 	else if (InItem->IsA<UPropertyAnimatorCoreBase>())
 	{
-		UPropertyAnimatorCoreBase* Animator = InItem->Get<UPropertyAnimatorCoreBase>();
+		UPropertyAnimatorCoreBase* Animator = InItem->Get<UPropertyAnimatorCoreBase>(0);
 
 		FBoolProperty* EnableProperty = FindFProperty<FBoolProperty>(UPropertyAnimatorCoreBase::StaticClass(), GET_MEMBER_NAME_CHECKED(UPropertyAnimatorCoreBase, bAnimatorEnabled));
 
 		const FSlateIcon ClassIcon = FSlateIconFinder::FindIconForClass(Animator->GetClass());
 
 		// Commands for item on key events
-		const TSharedPtr<FUICommandList> Commands = CreateAnimatorCommands(Animator);
-
-		// Action menu available in header in slim toolbar
-		static const FName HeaderModifierMenuName = TEXT("HeaderControllerMenu");
-		if (!UToolMenus::Get()->IsMenuRegistered(HeaderModifierMenuName))
-		{
-			UToolMenu* const HeaderModifierMenu = UToolMenus::Get()->RegisterMenu(HeaderModifierMenuName, NAME_None, EMultiBoxType::SlimHorizontalToolBar);
-			HeaderModifierMenu->AddDynamicSection(TEXT("FillHeaderControllerMenu"), FNewToolMenuDelegate::CreateUObject(this, &UPropertyAnimatorCoreEditorStackCustomization::FillAnimatorHeaderActionMenu));
-		}
-
-		// Context menu available when right clicking on item
-		static const FName ContextModifierMenuName = TEXT("ContextControllerMenu");
-		if (!UToolMenus::Get()->IsMenuRegistered(ContextModifierMenuName))
-		{
-			UToolMenu* const ContextModifierMenu = UToolMenus::Get()->RegisterMenu(ContextModifierMenuName, NAME_None, EMultiBoxType::Menu);
-			ContextModifierMenu->AddDynamicSection(TEXT("FillContextControllerMenu"), FNewToolMenuDelegate::CreateUObject(this, &UPropertyAnimatorCoreEditorStackCustomization::FillAnimatorContextActionMenu));
-		}
+		const TSharedPtr<FUICommandList> AnimatorCommands = CreateAnimatorCommands(InItem);
 
 		const TSet<FString> SearchKeywords
 		{
@@ -165,30 +225,40 @@ void UPropertyAnimatorCoreEditorStackCustomization::CustomizeItemHeader(const FO
 		/** Show last execution error messages if failed execution */
 		TAttribute<EOperatorStackEditorMessageType> MessageType = EOperatorStackEditorMessageType::None;
 		TAttribute<FText> MessageText = FText::GetEmpty();
+		FString HeaderLabel = TEXT("Animator");
 
-		TWeakObjectPtr<UPropertyAnimatorCoreBase> AnimatorWeak(Animator);
-
-		MessageType = TAttribute<EOperatorStackEditorMessageType>::CreateLambda([AnimatorWeak]()
+		if (InItem->GetValueCount() == 1)
 		{
-			if (const UPropertyAnimatorCoreBase* Animator = AnimatorWeak.Get())
+			HeaderLabel = Animator->GetAnimatorDisplayName();
+
+			TWeakObjectPtr<UPropertyAnimatorCoreBase> AnimatorWeak(Animator);
+
+			MessageType = TAttribute<EOperatorStackEditorMessageType>::CreateLambda([AnimatorWeak]()
 			{
-				return Animator->GetLinkedPropertiesCount() > 0
-					? EOperatorStackEditorMessageType::None
-					: EOperatorStackEditorMessageType::Info;
-			}
+				if (const UPropertyAnimatorCoreBase* Animator = AnimatorWeak.Get())
+				{
+					return Animator->GetLinkedPropertiesCount() > 0
+						? EOperatorStackEditorMessageType::None
+						: EOperatorStackEditorMessageType::Info;
+				}
 
-			return EOperatorStackEditorMessageType::None;
-		});
+				return EOperatorStackEditorMessageType::None;
+			});
 
-		MessageText = TAttribute<FText>::CreateLambda([AnimatorWeak]()
-		{
-			if (AnimatorWeak.IsValid())
+			MessageText = TAttribute<FText>::CreateLambda([AnimatorWeak]()
 			{
-				return LOCTEXT("NoPropertiesLinked", "No properties are currently linked to this animator");
-			}
+				if (AnimatorWeak.IsValid())
+				{
+					return LOCTEXT("NoPropertiesLinked", "No properties are currently linked to this animator");
+				}
 
-			return FText::GetEmpty();
-		});
+				return FText::GetEmpty();
+			});
+		}
+		else
+		{
+			HeaderLabel = Animator->GetAnimatorOriginalName().ToString();
+		}
 
 		static const FLinearColor AnimatorColor = FLinearColor(FColor::Orange).Desaturate(0.25);
 
@@ -198,11 +268,11 @@ void UPropertyAnimatorCoreEditorStackCustomization::CustomizeItemHeader(const FO
 			.SetSearchKeywords(SearchKeywords)
 			.SetExpandable(true)
 			.SetIcon(ClassIcon.GetIcon())
-			.SetLabel(FText::FromString(Animator->GetAnimatorDisplayName()))
+			.SetLabel(FText::FromString(HeaderLabel))
 			.SetProperty(EnableProperty)
-			.SetCommandList(Commands)
-			.SetToolbarMenu(HeaderModifierMenuName)
-			.SetContextMenu(ContextModifierMenuName)
+			.SetCommandList(AnimatorCommands)
+			.SetToolbarMenu(HeaderAnimatorMenuName)
+			.SetContextMenu(ContextAnimatorMenuName)
 			.SetMessageBox(MessageType, MessageText);
 	}
 
@@ -211,7 +281,7 @@ void UPropertyAnimatorCoreEditorStackCustomization::CustomizeItemHeader(const FO
 
 void UPropertyAnimatorCoreEditorStackCustomization::CustomizeItemBody(const FOperatorStackEditorItemPtr& InItem, const FOperatorStackEditorTree& InItemTree, FOperatorStackEditorBodyBuilder& InBodyBuilder)
 {
-	// Customize controller body
+	// Customize animator body
 	if (InItem->IsA<UPropertyAnimatorCoreBase>())
 	{
 		FBoolProperty* EnableProperty = FindFProperty<FBoolProperty>(UPropertyAnimatorCoreBase::StaticClass(), GET_MEMBER_NAME_CHECKED(UPropertyAnimatorCoreBase, bAnimatorEnabled));
@@ -239,9 +309,8 @@ void UPropertyAnimatorCoreEditorStackCustomization::CustomizeItemFooter(const FO
 	Super::CustomizeItemFooter(InItem, InItemTree, InFooterBuilder);
 }
 
-bool UPropertyAnimatorCoreEditorStackCustomization::OnIsItemDraggable(const FOperatorStackEditorItemPtr& InItem)
+bool UPropertyAnimatorCoreEditorStackCustomization::OnIsItemSelectable(const FOperatorStackEditorItemPtr& InItem)
 {
-	// Allow selection of item in listview for commands to work
 	if (InItem->IsA<UPropertyAnimatorCoreBase>())
 	{
 		return true;
@@ -255,9 +324,9 @@ const FSlateBrush* UPropertyAnimatorCoreEditorStackCustomization::GetIcon() cons
 	return FSlateIconFinder::FindIconForClass(UPropertyAnimatorCoreBase::StaticClass()).GetIcon();
 }
 
-void UPropertyAnimatorCoreEditorStackCustomization::RemoveControllerAction(UPropertyAnimatorCoreBase* InAnimator) const
+void UPropertyAnimatorCoreEditorStackCustomization::RemoveAnimatorAction(FOperatorStackEditorItemPtr InItem) const
 {
-	if (!InAnimator)
+	if (!InItem.IsValid() || !InItem->HasValue())
 	{
 		return;
 	}
@@ -269,24 +338,24 @@ void UPropertyAnimatorCoreEditorStackCustomization::RemoveControllerAction(UProp
 		return;
 	}
 
-	Subsystem->RemoveAnimator(InAnimator, true);
-}
+	TSet<UPropertyAnimatorCoreBase*> Animators;
 
-void UPropertyAnimatorCoreEditorStackCustomization::RemoveControllersAction(UPropertyAnimatorCoreComponent* InComponent) const
-{
-	const UPropertyAnimatorCoreSubsystem* Subsystem = UPropertyAnimatorCoreSubsystem::Get();
-	if (!InComponent || !Subsystem)
+	if (InItem->IsA<UPropertyAnimatorCoreBase>())
 	{
-		return;
+		Animators.Append(InItem->GetAsArray<UPropertyAnimatorCoreBase>());
 	}
-
-	TSet<UPropertyAnimatorCoreBase*> Animators;
-	for (const TObjectPtr<UPropertyAnimatorCoreBase>& Animator : InComponent->GetAnimators())
+	else if (InItem->IsA<UPropertyAnimatorCoreComponent>())
 	{
-		Animators.Add(Animator);
+		for (const UPropertyAnimatorCoreComponent* AnimatorComponent : InItem->GetAsArray<UPropertyAnimatorCoreComponent>())
+		{
+			Animators.Append(AnimatorComponent->GetAnimators());
+		}
 	}
 
-	Subsystem->RemoveAnimators(Animators, true);
+	if (!Subsystem->RemoveAnimators(Animators, /** Transact */true))
+	{
+		UE_LOG(LogPropertyAnimatorCoreEditorStackCustomization, Warning, TEXT("Could not remove %i animator(s)"), Animators.Num())
+	}
 }
 
 void UPropertyAnimatorCoreEditorStackCustomization::FillAddAnimatorMenuSection(UToolMenu* InToolMenu) const
@@ -319,7 +388,10 @@ void UPropertyAnimatorCoreEditorStackCustomization::FillAddAnimatorMenuSection(U
 	{
 		if (ContextItem->IsA<UObject>())
 		{
-			ContextObjects.Add(ContextItem->Get<UObject>());
+			for (UObject* ContextObject : ContextItem->GetAsArray<UObject>())
+			{
+				ContextObjects.Add(ContextObject);
+			}
 		}
 	}
 
@@ -328,43 +400,6 @@ void UPropertyAnimatorCoreEditorStackCustomization::FillAddAnimatorMenuSection(U
 	AnimatorEditorSubsystem->FillAnimatorMenu(InToolMenu, MenuContext, MenuOptions);
 }
 
-void UPropertyAnimatorCoreEditorStackCustomization::FillComponentHeaderActionMenu(UToolMenu* InToolMenu)
-{
-	if (!InToolMenu)
-	{
-		return;
-	}
-
-	const UOperatorStackEditorMenuContext* const MenuContext = InToolMenu->FindContext<UOperatorStackEditorMenuContext>();
-	if (!MenuContext)
-	{
-		return;
-	}
-
-	const FOperatorStackEditorItemPtr ItemContext = MenuContext->GetItem();
-	if (!ItemContext)
-	{
-		return;
-	}
-
-	UPropertyAnimatorCoreComponent* ControllerComponent = ItemContext->Get<UPropertyAnimatorCoreComponent>();
-	if (!IsValid(ControllerComponent))
-	{
-		return;
-	}
-
-	// Add remove controllers entry
-	const FToolMenuEntry RemoveControllersAction = FToolMenuEntry::InitToolBarButton(
-		TEXT("RemoveControllersMenuEntry")
-		, FExecuteAction::CreateUObject(this, &UPropertyAnimatorCoreEditorStackCustomization::RemoveControllersAction, ControllerComponent)
-		, FText::GetEmpty()
-		, FText::GetEmpty()
-		, FSlateIcon(FAppStyle::GetAppStyleSetName(), "Icons.Delete")
-	);
-
-	InToolMenu->AddMenuEntry(RemoveControllersAction.Name, RemoveControllersAction);
-}
-
 void UPropertyAnimatorCoreEditorStackCustomization::FillAnimatorHeaderActionMenu(UToolMenu* InToolMenu)
 {
 	if (!InToolMenu)
@@ -384,22 +419,20 @@ void UPropertyAnimatorCoreEditorStackCustomization::FillAnimatorHeaderActionMenu
 		return;
 	}
 
-	UPropertyAnimatorCoreBase* Controller = ItemContext->Get<UPropertyAnimatorCoreBase>();
-	if (!IsValid(Controller))
+	if (!ItemContext->IsA<UPropertyAnimatorCoreBase>() && !ItemContext->IsA<UPropertyAnimatorCoreComponent>())
 	{
 		return;
 	}
 
-	// Add remove controller entry
-	const FToolMenuEntry RemoveControllerAction = FToolMenuEntry::InitToolBarButton(
-		TEXT("RemoveControllerMenuEntry")
-		, FExecuteAction::CreateUObject(this, &UPropertyAnimatorCoreEditorStackCustomization::RemoveControllerAction, Controller)
+	const FToolMenuEntry RemoveAnimatorAction = FToolMenuEntry::InitToolBarButton(
+		TEXT("RemoveAnimatorMenuEntry")
+		, FExecuteAction::CreateUObject(this, &UPropertyAnimatorCoreEditorStackCustomization::RemoveAnimatorAction, ItemContext)
 		, FText::GetEmpty()
 		, FText::GetEmpty()
 		, FSlateIcon(FAppStyle::GetAppStyleSetName(), "Icons.Delete")
 	);
 
-	InToolMenu->AddMenuEntry(RemoveControllerAction.Name, RemoveControllerAction);
+	InToolMenu->AddMenuEntry(RemoveAnimatorAction.Name, RemoveAnimatorAction);
 }
 
 void UPropertyAnimatorCoreEditorStackCustomization::FillAnimatorContextActionMenu(UToolMenu* InToolMenu) const
@@ -421,8 +454,7 @@ void UPropertyAnimatorCoreEditorStackCustomization::FillAnimatorContextActionMen
 		return;
 	}
 
-	const UPropertyAnimatorCoreBase* Controller = ItemContext->Get<UPropertyAnimatorCoreBase>();
-	if (!IsValid(Controller))
+	if (!ItemContext->IsA<UPropertyAnimatorCoreBase>() && !ItemContext->IsA<UPropertyAnimatorCoreComponent>())
 	{
 		return;
 	}
@@ -431,18 +463,18 @@ void UPropertyAnimatorCoreEditorStackCustomization::FillAnimatorContextActionMen
 	TSharedPtr<const FUICommandList> Commands;
 	const TSharedPtr<FUICommandInfo> DeleteCommand = FGenericCommands::Get().Delete;
 	InToolMenu->Context.GetActionForCommand(DeleteCommand, Commands);
-	const FToolMenuEntry RemoveControllerMenuEntry = FToolMenuEntry::InitMenuEntryWithCommandList(DeleteCommand, Commands);
+	const FToolMenuEntry RemoveAnimatorMenuEntry = FToolMenuEntry::InitMenuEntryWithCommandList(DeleteCommand, Commands);
 
-	InToolMenu->AddMenuEntry(RemoveControllerMenuEntry.Name, RemoveControllerMenuEntry);
+	InToolMenu->AddMenuEntry(RemoveAnimatorMenuEntry.Name, RemoveAnimatorMenuEntry);
 }
 
-TSharedRef<FUICommandList> UPropertyAnimatorCoreEditorStackCustomization::CreateAnimatorCommands(UPropertyAnimatorCoreBase* InAnimator)
+TSharedRef<FUICommandList> UPropertyAnimatorCoreEditorStackCustomization::CreateAnimatorCommands(FOperatorStackEditorItemPtr InItem)
 {
 	TSharedRef<FUICommandList> Commands = MakeShared<FUICommandList>();
 
 	Commands->MapAction(FGenericCommands::Get().Delete,
 		FUIAction(
-			FExecuteAction::CreateUObject(this, &UPropertyAnimatorCoreEditorStackCustomization::RemoveControllerAction, InAnimator),
+			FExecuteAction::CreateUObject(this, &UPropertyAnimatorCoreEditorStackCustomization::RemoveAnimatorAction, InItem),
 			FCanExecuteAction()
 		)
 	);
diff --git a/Engine/Plugins/Experimental/PropertyAnimatorCore/Source/PropertyAnimatorCoreEditor/Private/Customizations/PropertyAnimatorCoreEditorStackCustomization.h b/Engine/Plugins/Experimental/PropertyAnimatorCore/Source/PropertyAnimatorCoreEditor/Private/Customizations/PropertyAnimatorCoreEditorStackCustomization.h
--- a/Engine/Plugins/Experimental/PropertyAnimatorCore/Source/PropertyAnimatorCoreEditor/Private/Customizations/PropertyAnimatorCoreEditorStackCustomization.h
+++ b/Engine/Plugins/Experimental/PropertyAnimatorCore/Source/PropertyAnimatorCoreEditor/Private/Customizations/PropertyAnimatorCoreEditorStackCustomization.h
@@ -20,26 +20,25 @@ public:
 	UPropertyAnimatorCoreEditorStackCustomization();
 
 	//~ Begin UOperatorStackEditorStackCustomization
-	virtual bool TransformContextItem(const FOperatorStackEditorItemPtr& InItem, TArray<FOperatorStackEditorItemPtr>& OutTransformedItems) const override;
+	virtual bool GetRootItem(const FOperatorStackEditorContext& InContext, FOperatorStackEditorItemPtr& OutRootItem) const override;
+	virtual bool GetChildrenItem(const FOperatorStackEditorItemPtr& InItem, TArray<FOperatorStackEditorItemPtr>& OutChildrenItems) const override;
 	virtual void CustomizeStackHeader(const FOperatorStackEditorTree& InItemTree, FOperatorStackEditorHeaderBuilder& InHeaderBuilder) override;
 	virtual void CustomizeItemHeader(const FOperatorStackEditorItemPtr& InItem, const FOperatorStackEditorTree& InItemTree, FOperatorStackEditorHeaderBuilder& InHeaderBuilder) override;
 	virtual void CustomizeItemBody(const FOperatorStackEditorItemPtr& InItem, const FOperatorStackEditorTree& InItemTree, FOperatorStackEditorBodyBuilder& InBodyBuilder) override;
 	virtual void CustomizeItemFooter(const FOperatorStackEditorItemPtr& InItem, const FOperatorStackEditorTree& InItemTree, FOperatorStackEditorFooterBuilder& InFooterBuilder) override;
-	virtual bool OnIsItemDraggable(const FOperatorStackEditorItemPtr& InItem) override;
+	virtual bool OnIsItemSelectable(const FOperatorStackEditorItemPtr& InItem) override;
 	virtual const FSlateBrush* GetIcon() const override;
 	//~ End UOperatorStackEditorStackCustomization
 
 protected:
 	/** Remove animator menu action */
-	void RemoveControllerAction(UPropertyAnimatorCoreBase* InAnimator) const;
-	void RemoveControllersAction(UPropertyAnimatorCoreComponent* InComponent) const;
+	void RemoveAnimatorAction(FOperatorStackEditorItemPtr InItem) const;
 
 	/** Fill item action menus */
 	void FillAddAnimatorMenuSection(UToolMenu* InToolMenu) const;
-	void FillComponentHeaderActionMenu(UToolMenu* InToolMenu);
 	void FillAnimatorHeaderActionMenu(UToolMenu* InToolMenu);
 	void FillAnimatorContextActionMenu(UToolMenu* InToolMenu) const;
 
 	/** Create commands for animator */
-	TSharedRef<FUICommandList> CreateAnimatorCommands(UPropertyAnimatorCoreBase* InAnimator);
+	TSharedRef<FUICommandList> CreateAnimatorCommands(FOperatorStackEditorItemPtr InItem);
 };
\ No newline at end of file
