MotionDesign : Modifiers
- Fixed property is not initialized properly

diff --git a/Engine/Plugins/Experimental/ActorModifierCore/Source/ActorModifierCore/Public/Modifiers/Utilities/ActorModifierCoreLibraryDefs.h b/Engine/Plugins/Experimental/ActorModifierCore/Source/ActorModifierCore/Public/Modifiers/Utilities/ActorModifierCoreLibraryDefs.h
--- a/Engine/Plugins/Experimental/ActorModifierCore/Source/ActorModifierCore/Public/Modifiers/Utilities/ActorModifierCoreLibraryDefs.h
+++ b/Engine/Plugins/Experimental/ActorModifierCore/Source/ActorModifierCore/Public/Modifiers/Utilities/ActorModifierCoreLibraryDefs.h
@@ -37,7 +37,7 @@ struct FActorModifierCoreCloneOperation
 	GENERATED_BODY()
 
 	UPROPERTY(BlueprintReadWrite, Category="Motion Design|Modifiers|Utility")
-	UActorModifierCoreBase* CloneModifier;
+	UActorModifierCoreBase* CloneModifier = nullptr;
 
 	UPROPERTY(BlueprintReadWrite, Category="Motion Design|Modifiers|Utility")
 	EActorModifierCoreStackPosition ClonePosition = EActorModifierCoreStackPosition::Before;
@@ -52,7 +52,7 @@ struct FActorModifierCoreMoveOperation
 	GENERATED_BODY()
 
 	UPROPERTY(BlueprintReadWrite, Category="Motion Design|Modifiers|Utility")
-	UActorModifierCoreBase* MoveModifier;
+	UActorModifierCoreBase* MoveModifier = nullptr;
 
 	UPROPERTY(BlueprintReadWrite, Category="Motion Design|Modifiers|Utility")
 	EActorModifierCoreStackPosition MovePosition = EActorModifierCoreStackPosition::Before;
@@ -67,7 +67,7 @@ struct FActorModifierCoreRemoveOperation
 	GENERATED_BODY()
 
 	UPROPERTY(BlueprintReadWrite, Category="Motion Design|Modifiers|Utility")
-	UActorModifierCoreBase* RemoveModifier;
+	UActorModifierCoreBase* RemoveModifier = nullptr;
 
 	UPROPERTY(BlueprintReadWrite, Category="Motion Design|Modifiers|Utility")
 	bool bRemoveDependencies = false;
