MotionDesign : ActorModifierCore
- Add early support for blueprintable modifiers (internal use only for now)

diff --git a/Engine/Plugins/Experimental/ActorModifierCore/ActorModifierCore.uplugin b/Engine/Plugins/Experimental/ActorModifierCore/ActorModifierCore.uplugin
--- a/Engine/Plugins/Experimental/ActorModifierCore/ActorModifierCore.uplugin
+++ b/Engine/Plugins/Experimental/ActorModifierCore/ActorModifierCore.uplugin
@@ -25,6 +25,11 @@
 			"Type": "Runtime",
 			"LoadingPhase": "Default"
 		},
+		{
+			"Name": "ActorModifierCoreBlueprint",
+			"Type": "UncookedOnly",
+			"LoadingPhase": "Default"
+		},
 		{
 			"Name": "ActorModifierCoreEditor",
 			"Type": "Editor",
diff --git a/Engine/Plugins/Experimental/ActorModifierCore/Source/ActorModifierCore/Private/Modifiers/ActorModifierCoreBase.cpp b/Engine/Plugins/Experimental/ActorModifierCore/Source/ActorModifierCore/Private/Modifiers/ActorModifierCoreBase.cpp
--- a/Engine/Plugins/Experimental/ActorModifierCore/Source/ActorModifierCore/Private/Modifiers/ActorModifierCoreBase.cpp
+++ b/Engine/Plugins/Experimental/ActorModifierCore/Source/ActorModifierCore/Private/Modifiers/ActorModifierCoreBase.cpp
@@ -2,6 +2,7 @@
 
 #include "Modifiers/ActorModifierCoreBase.h"
 
+#include "Async/TaskGraphInterfaces.h"
 #include "Modifiers/ActorModifierCoreComponent.h"
 #include "Modifiers/ActorModifierCoreSharedObject.h"
 #include "Modifiers/ActorModifierCoreStack.h"
@@ -110,17 +111,39 @@ void UActorModifierCoreBase::Apply()
 
 void UActorModifierCoreBase::Next()
 {
-	if (!bModifierIdle && ApplyPromise.IsValid())
+	auto ExecuteNext = [this]()
+	{
+		if (!bModifierIdle && ApplyPromise.IsValid())
+		{
+			// Success
+			Status = FActorModifierCoreStatus(EActorModifierCoreStatus::Success, FText::GetEmpty());
+			ApplyPromise->SetValue(true);
+		}
+		else
+		{
+			LogModifier(TEXT("Next is called again after execution is done"), true);
+			checkNoEntry()
+		}
+	};
+
+	if (!IsInGameThread())
 	{
-		// Success
-		Status = FActorModifierCoreStatus(EActorModifierCoreStatus::Success, FText::GetEmpty());
+		TWeakObjectPtr<UActorModifierCoreBase> ThisWeak(this);
+		FFunctionGraphTask::CreateAndDispatchWhenReady([ThisWeak, ExecuteNext]()
+		{
+			const UActorModifierCoreBase* This = ThisWeak.Get();
+
+			if (!This)
+			{
+				return;
+			}
 
-		ApplyPromise->SetValue(true);
+			ExecuteNext();
+		}, {}, nullptr, ENamedThreads::GameThread);
 	}
 	else
 	{
-		LogModifier(TEXT("Next is called again after execution is done"), true);
-		checkNoEntry()
+		ExecuteNext();
 	}
 }
 
@@ -151,6 +174,7 @@ void UActorModifierCoreBase::Unapply()
 	if (bModifierApplied)
 	{
 		RestorePreState();
+
 		bModifierApplied = false;
 	}
 }
@@ -436,6 +460,12 @@ void UActorModifierCoreBase::DeferInitializeModifier()
 
 	// Begin batch operation to avoid updating every time a modifier is loaded
 	UActorModifierCoreStack* Stack = GetRootModifierStack();
+
+	if (!Stack)
+	{
+		return;
+	}
+
 	if (!Stack->IsModifierExecutionLocked() && !Stack->IsModifierStackInitialized())
 	{
 		Stack->LockModifierExecution();
@@ -611,8 +641,9 @@ void UActorModifierCoreBase::LogModifier(const FString& InLog, bool bInForce) co
 	{
 		const FString ActorLabel = GetModifiedActor() ? *GetModifiedActor()->GetActorNameOrLabel() : TEXT("Invalid actor");
 		const FString ModifierLabel = GetModifierName().ToString();
+		const FString ClassLabel = GetClass()->GetName();
 
-		UE_LOG(LogActorModifierCoreBase, Log, TEXT("[%s][%s] %s"), *ActorLabel, *ModifierLabel, *InLog);
+		UE_LOG(LogActorModifierCoreBase, Log, TEXT("[%s][%s][%s] %s"), *ActorLabel, *ClassLabel, *ModifierLabel, *InLog);
 	}
 }
 
@@ -746,6 +777,12 @@ void UActorModifierCoreBase::InitializeModifier(EActorModifierCoreEnableReason I
 		// set new actor
 		ModifiedActor = GetModifiedActor();
 
+		// if the metadata is not initialized, reload it from CDO
+		if (!Metadata.IsValid())
+		{
+			PostModifierCreation(GetModifierStack());
+		}
+
 		// Initialize profiler
 		if (!Profiler.IsValid())
 		{
diff --git a/Engine/Plugins/Experimental/ActorModifierCore/Source/ActorModifierCore/Private/Modifiers/ActorModifierCoreDefs.cpp b/Engine/Plugins/Experimental/ActorModifierCore/Source/ActorModifierCore/Private/Modifiers/ActorModifierCoreDefs.cpp
--- a/Engine/Plugins/Experimental/ActorModifierCore/Source/ActorModifierCore/Private/Modifiers/ActorModifierCoreDefs.cpp
+++ b/Engine/Plugins/Experimental/ActorModifierCore/Source/ActorModifierCore/Private/Modifiers/ActorModifierCoreDefs.cpp
@@ -65,6 +65,11 @@ bool FActorModifierCoreMetadata::IsAllowedBefore(const FName& InModifierName) co
 
 bool FActorModifierCoreMetadata::IsCompatibleWith(const AActor* InActor) const
 {
+	if (CompatibilityRuleDelegate.IsBound())
+	{
+		return CompatibilityRuleDelegate.Execute(InActor);
+	}
+
 	return CompatibilityRuleFunction(InActor);
 }
 
@@ -131,6 +136,7 @@ bool FActorModifierCoreMetadata::ResetDefault()
 		bTickAllowed = CDOMetadata.bTickAllowed;
 		bMultipleAllowed = CDOMetadata.bMultipleAllowed;
 		CompatibilityRuleFunction = CDOMetadata.CompatibilityRuleFunction;
+		CompatibilityRuleDelegate = CDOMetadata.CompatibilityRuleDelegate;
 		ProfilerFunction = CDOMetadata.ProfilerFunction;
 
 		return true;
@@ -246,6 +252,12 @@ FActorModifierCoreMetadata& FActorModifierCoreMetadata::SetCompatibilityRule(con
 	return *this;
 }
 
+FActorModifierCoreMetadata& FActorModifierCoreMetadata::SetCompatibilityRule(const FModifierCompatibilityRule& InModifierRule)
+{
+	CompatibilityRuleDelegate = InModifierRule;
+	return *this;
+}
+
 UActorModifierCoreBase* FActorModifierCoreMetadata::CreateModifierInstance(UActorModifierCoreStack* InStack) const
 {
 	if (InStack && InStack->GetModifiedActor())
diff --git a/Engine/Plugins/Experimental/ActorModifierCore/Source/ActorModifierCore/Private/Modifiers/ActorModifierCoreStack.cpp b/Engine/Plugins/Experimental/ActorModifierCore/Source/ActorModifierCore/Private/Modifiers/ActorModifierCoreStack.cpp
--- a/Engine/Plugins/Experimental/ActorModifierCore/Source/ActorModifierCore/Private/Modifiers/ActorModifierCoreStack.cpp
+++ b/Engine/Plugins/Experimental/ActorModifierCore/Source/ActorModifierCore/Private/Modifiers/ActorModifierCoreStack.cpp
@@ -5,6 +5,7 @@
 #include "Algo/RemoveIf.h"
 #include "Components/SceneComponent.h"
 #include "Modifiers/ActorModifierCoreComponent.h"
+#include "Modifiers/Blueprints/ActorModifierCoreBlueprintBase.h"
 #include "Subsystems/ActorModifierCoreSubsystem.h"
 
 #define LOCTEXT_NAMESPACE "ActorModifierCoreStack"
@@ -14,6 +15,7 @@ DEFINE_LOG_CATEGORY_STATIC(LogActorModifierCoreStack, Log, All);
 UActorModifierCoreStack::FOnModifierUpdated UActorModifierCoreStack::OnModifierAddedDelegate;
 UActorModifierCoreStack::FOnModifierUpdated UActorModifierCoreStack::OnModifierRemovedDelegate;
 UActorModifierCoreStack::FOnModifierUpdated UActorModifierCoreStack::OnModifierMovedDelegate;
+UActorModifierCoreStack::FOnModifierUpdated UActorModifierCoreStack::OnModifierReplacedDelegate;
 
 UActorModifierCoreStack* UActorModifierCoreStack::Create(UActorModifierCoreComponent* InComponent, UActorModifierCoreStack* InParentStack)
 {
@@ -101,25 +103,31 @@ void UActorModifierCoreStack::OnModifierAdded(EActorModifierCoreEnableReason InR
 {
 	Super::OnModifierAdded(InReason);
 
-	// initialize only root stack
-	if (IsRootStack())
+	if (UActorModifierCoreSubsystem* ModifierSubsystem = UActorModifierCoreSubsystem::Get())
 	{
-		// register with subsystem for fast query
-		if (UActorModifierCoreSubsystem::Get()->RegisterActorModifierStack(this))
+		// initialize only root stack
+		if (IsRootStack())
 		{
-			// initialize events when actor changes
-			if (AActor* Actor = GetModifiedActor())
+			// register with subsystem for fast query
+			if (ModifierSubsystem->RegisterActorModifierStack(this))
 			{
-				Actor->OnDestroyed.AddUniqueDynamic(this, &UActorModifierCoreStack::OnActorDestroyed);
-				if (USceneComponent* RootComponent = ModifiedActor->GetRootComponent())
+				// initialize events when actor changes
+				if (AActor* Actor = GetModifiedActor())
 				{
-					if (!RootComponent->TransformUpdated.IsBoundToObject(this))
+					Actor->OnDestroyed.AddUniqueDynamic(this, &UActorModifierCoreStack::OnActorDestroyed);
+
+					if (USceneComponent* RootComponent = ModifiedActor->GetRootComponent())
 					{
-						RootComponent->TransformUpdated.AddUObject(this, &UActorModifierCoreStack::OnActorTransformUpdated);
+						if (!RootComponent->TransformUpdated.IsBoundToObject(this))
+						{
+							RootComponent->TransformUpdated.AddUObject(this, &UActorModifierCoreStack::OnActorTransformUpdated);
+						}
 					}
 				}
 			}
 		}
+
+		ModifierSubsystem->OnModifierReplaced().AddUObject(this, &UActorModifierCoreStack::OnBlueprintModifierReplaced);
 	}
 
 	// initialize inner modifiers if they are not initialized
@@ -1588,6 +1596,59 @@ void UActorModifierCoreStack::CheckModifierOptimization(bool bInInvalidateAll)
 	bModifierOptimized = UnoptimizedModifiers.IsEmpty();
 }
 
+void UActorModifierCoreStack::OnBlueprintModifierReplaced(UActorModifierCoreBlueprintBase* InOldModifier, UActorModifierCoreBlueprintBase* InNewModifier)
+{
+	const int32 ModifierIndex = Modifiers.Find(InOldModifier);
+
+	if (ModifierIndex == INDEX_NONE)
+	{
+		return;
+	}
+
+	// Init metadata for new modifier
+	InNewModifier->PostModifierCreation(this);
+
+	// Copy non uproperty var
+	InNewModifier->bModifierApplied = InOldModifier->bModifierApplied;
+	InNewModifier->bModifierInitialized = InOldModifier->bModifierInitialized;
+	InNewModifier->bModifierOptimized = InOldModifier->bModifierOptimized;
+
+	AActor* TargetActor = GetModifiedActor();
+
+	// Check if support rule has changed
+	const bool bActorSupported = InNewModifier->GetModifierMetadata().IsCompatibleWith(TargetActor);
+
+	if (bActorSupported)
+	{
+		Modifiers[ModifierIndex] = InNewModifier;
+	}
+	else
+	{
+		Modifiers.RemoveAt(ModifierIndex);
+	}
+
+	// Replace old modifier by new modifier in stack
+	const int32 ExecuteIndex = ExecuteModifiers.Find(InOldModifier);
+
+	if (ExecuteModifiers.IsValidIndex(ExecuteIndex))
+	{
+		ExecuteModifiers[ExecuteIndex] = InNewModifier;
+	}
+
+	const int32 CurrentIndex = CurrentModifiers.Find(InOldModifier);
+
+	if (CurrentModifiers.IsValidIndex(CurrentIndex))
+	{
+		CurrentModifiers[CurrentIndex] = InNewModifier;
+	}
+
+	InNewModifier->OnModifierReplacedEvent(TargetActor);
+
+	OnModifierReplacedDelegate.Broadcast(InNewModifier);
+
+	InNewModifier->MarkModifierDirty();
+}
+
 void UActorModifierCoreStack::RestorePreState()
 {
 	TArray<TObjectPtr<UActorModifierCoreBase>> ModifiersRestoreChain;
@@ -1645,24 +1706,26 @@ void UActorModifierCoreStack::Apply()
 
 void UActorModifierCoreStack::OnModifierRemoved(EActorModifierCoreDisableReason InReason)
 {
-	if (IsRootStack())
+	if (UActorModifierCoreSubsystem* Subsystem = UActorModifierCoreSubsystem::Get())
 	{
-		if (UActorModifierCoreSubsystem* Subsystem = UActorModifierCoreSubsystem::Get())
+		if (IsRootStack())
 		{
-			if (const AActor* ConstStackActor = Subsystem->GetModifierStackActor(this))
+			if (AActor* TargetActor = GetModifiedActor())
 			{
 				// unregister stack from actor in subsystem
-				AActor* StackActor = const_cast<AActor*>(ConstStackActor);
-				Subsystem->UnregisterActorModifierStack(StackActor);
+				Subsystem->UnregisterActorModifierStack(TargetActor);
 
 				// unbind actor event
-				StackActor->OnDestroyed.RemoveAll(this);
-				if (USceneComponent* RootComponent = StackActor->GetRootComponent())
+				TargetActor->OnDestroyed.RemoveAll(this);
+
+				if (USceneComponent* RootComponent = TargetActor->GetRootComponent())
 				{
 					RootComponent->TransformUpdated.RemoveAll(this);
 				}
 			}
 		}
+
+		Subsystem->OnModifierReplaced().RemoveAll(this);
 	}
 
 	// un-initialize inner modifiers if they are not un-initialized
diff --git a/Engine/Plugins/Experimental/ActorModifierCore/Source/ActorModifierCore/Private/Modifiers/Blueprints/ActorModifierCoreBlueprintBase.cpp b/Engine/Plugins/Experimental/ActorModifierCore/Source/ActorModifierCore/Private/Modifiers/Blueprints/ActorModifierCoreBlueprintBase.cpp
new file mode 100644
--- /dev/null
+++ b/Engine/Plugins/Experimental/ActorModifierCore/Source/ActorModifierCore/Private/Modifiers/Blueprints/ActorModifierCoreBlueprintBase.cpp
@@ -0,0 +1,155 @@
+// Copyright Epic Games, Inc. All Rights Reserved.
+
+#include "Modifiers/Blueprints/ActorModifierCoreBlueprintBase.h"
+
+#define LOCTEXT_NAMESPACE "ActorModifierCoreBlueprintBase"
+
+void UActorModifierCoreBlueprintBase::FlagModifierDirty()
+{
+	if (!IsTemplate())
+	{
+		LogModifier(TEXT("Blueprint modifier flagged dirty"));
+
+		MarkModifierDirty(/** Execute */true);
+	}
+}
+
+#if WITH_EDITOR
+void UActorModifierCoreBlueprintBase::PostEditChangeProperty(FPropertyChangedEvent& InPropertyChangedEvent)
+{
+	Super::PostEditChangeProperty(InPropertyChangedEvent);
+
+	// Only trigger modifier if not interactive and property owned by children of this class
+	if (InPropertyChangedEvent.Property
+		&& InPropertyChangedEvent.Property->GetOwnerClass()
+		&& InPropertyChangedEvent.Property->GetOwnerClass()->IsChildOf<UActorModifierCoreBlueprintBase>()
+		&& InPropertyChangedEvent.ChangeType != EPropertyChangeType::Interactive)
+	{
+		FlagModifierDirty();
+	}
+}
+#endif
+
+void UActorModifierCoreBlueprintBase::OnModifierCDOSetup(FActorModifierCoreMetadata& InMetadata)
+{
+	Super::OnModifierCDOSetup(InMetadata);
+
+	LogModifier(TEXT("Blueprint modifier setup"));
+
+	// Since the blueprint utility functions returns a copy of the struct, make sure we get a result
+	OnModifierSetupEvent(InMetadata, /** OutMetadata */InMetadata);
+
+	if (InMetadata.GetName().IsNone())
+	{
+		LogModifier(TEXT("Blueprint modifier setup failed : Name was not defined"), true);
+	}
+}
+
+void UActorModifierCoreBlueprintBase::OnModifierAdded(EActorModifierCoreEnableReason InReason)
+{
+	Super::OnModifierAdded(InReason);
+
+	AActor* TargetActor = GetModifiedActor();
+
+	if (IsValid(TargetActor))
+	{
+		LogModifier(FString::Printf(TEXT("Blueprint modifier added with reason %i"), InReason));
+
+		OnModifierAddedEvent(TargetActor, InReason);
+	}
+}
+
+void UActorModifierCoreBlueprintBase::OnModifierEnabled(EActorModifierCoreEnableReason InReason)
+{
+	Super::OnModifierEnabled(InReason);
+
+	AActor* TargetActor = GetModifiedActor();
+
+	if (IsValid(TargetActor))
+	{
+		LogModifier(FString::Printf(TEXT("Blueprint modifier enabled with reason %i"), InReason));
+
+		OnModifierEnabledEvent(TargetActor, InReason);
+	}
+}
+
+void UActorModifierCoreBlueprintBase::OnModifierDisabled(EActorModifierCoreDisableReason InReason)
+{
+	Super::OnModifierDisabled(InReason);
+
+	AActor* TargetActor = GetModifiedActor();
+
+	if (IsValid(TargetActor))
+	{
+		LogModifier(FString::Printf(TEXT("Blueprint modifier disabled with reason %i"), InReason));
+
+		OnModifierDisabledEvent(TargetActor, InReason);
+	}
+}
+
+void UActorModifierCoreBlueprintBase::OnModifierRemoved(EActorModifierCoreDisableReason InReason)
+{
+	Super::OnModifierRemoved(InReason);
+
+	AActor* TargetActor = GetModifiedActor();
+
+	if (IsValid(TargetActor))
+	{
+		LogModifier(FString::Printf(TEXT("Blueprint modifier removed with reason %i"), InReason));
+
+		OnModifierRemovedEvent(TargetActor, InReason);
+	}
+}
+
+void UActorModifierCoreBlueprintBase::SavePreState()
+{
+	Super::SavePreState();
+
+	AActor* TargetActor = GetModifiedActor();
+
+	if (IsValid(TargetActor))
+	{
+		LogModifier(TEXT("Blueprint modifier save pre state"));
+
+		OnModifierSaveStateEvent(TargetActor);
+	}
+}
+
+void UActorModifierCoreBlueprintBase::RestorePreState()
+{
+	Super::RestorePreState();
+
+	AActor* TargetActor = GetModifiedActor();
+
+	if (IsValid(TargetActor))
+	{
+		LogModifier(TEXT("Blueprint modifier restore pre state"));
+
+		OnModifierRestoreStateEvent(TargetActor);
+	}
+}
+
+void UActorModifierCoreBlueprintBase::Apply()
+{
+	FText FailReason = FText::GetEmpty();
+	AActor* TargetActor = GetModifiedActor();
+
+	LogModifier(TEXT("Blueprint modifier apply"));
+
+	if (IsValid(TargetActor) && OnModifierApplyEvent(TargetActor, FailReason))
+	{
+		Next();
+	}
+	else
+	{
+		// Fail reason must be set
+		if (FailReason.IsEmpty())
+		{
+			FailReason = FText::Format(LOCTEXT("ApplyFailed", "{0} : Blueprint modifier {1} apply failed"), FText::FromString(TargetActor ? TargetActor->GetActorNameOrLabel() : TEXT("?")), FText::FromName(GetModifierName()));
+		}
+
+		Fail(FailReason);
+	}
+}
+
+#undef LOCTEXT_NAMESPACE
diff --git a/Engine/Plugins/Experimental/ActorModifierCore/Source/ActorModifierCore/Private/Modifiers/Utilities/ActorModifierCoreLibrary.cpp b/Engine/Plugins/Experimental/ActorModifierCore/Source/ActorModifierCore/Private/Modifiers/Utilities/ActorModifierCoreLibrary.cpp
--- a/Engine/Plugins/Experimental/ActorModifierCore/Source/ActorModifierCore/Private/Modifiers/Utilities/ActorModifierCoreLibrary.cpp
+++ b/Engine/Plugins/Experimental/ActorModifierCore/Source/ActorModifierCore/Private/Modifiers/Utilities/ActorModifierCoreLibrary.cpp
@@ -497,3 +497,55 @@ bool UActorModifierCoreLibrary::GetAvailableModifiers(TSet<TSubclassOf<UActorMod
 
 	return true;
 }
+
+FActorModifierCoreMetadata& UActorModifierCoreLibrary::SetModifierMetadataName(FActorModifierCoreMetadata& InMetadata, FName InName)
+{
+	InMetadata.SetName(InName);
+	return InMetadata;
+}
+
+FActorModifierCoreMetadata& UActorModifierCoreLibrary::SetModifierMetadataCategory(FActorModifierCoreMetadata& InMetadata, FName InCategory)
+{
+	InMetadata.SetCategory(InCategory);
+	return InMetadata;
+}
+
+FActorModifierCoreMetadata& UActorModifierCoreLibrary::SetModifierMetadataDisplayName(FActorModifierCoreMetadata& InMetadata, const FText& InName)
+{
+#if WITH_EDITOR
+	InMetadata.SetDisplayName(InName);
+#endif
+	return InMetadata;
+}
+
+FActorModifierCoreMetadata& UActorModifierCoreLibrary::SetModifierMetadataColor(FActorModifierCoreMetadata& InMetadata, const FLinearColor& InColor)
+{
+#if WITH_EDITOR
+	InMetadata.SetColor(InColor);
+#endif
+	return InMetadata;
+}
+
+FActorModifierCoreMetadata& UActorModifierCoreLibrary::SetModifierMetadataDescription(FActorModifierCoreMetadata& InMetadata, const FText& InDescription)
+{
+#if WITH_EDITOR
+	InMetadata.SetDescription(InDescription);
+#endif
+	return InMetadata;
+}
+
+FActorModifierCoreMetadata& UActorModifierCoreLibrary::AddModifierMetadataDependency(FActorModifierCoreMetadata& InMetadata, TSubclassOf<UActorModifierCoreBase> InModifierClass)
+{
+	if (const UActorModifierCoreBase* Modifier = InModifierClass.GetDefaultObject())
+	{
+		InMetadata.AddDependency(Modifier->GetModifierName());
+	}
+
+	return InMetadata;
+}
+
+FActorModifierCoreMetadata& UActorModifierCoreLibrary::SetModifierMetadataCompatibilityRule(FActorModifierCoreMetadata& InMetadata, const FModifierCompatibilityRule& InDelegate)
+{
+	InMetadata.SetCompatibilityRule(InDelegate);
+	return InMetadata;
+}
diff --git a/Engine/Plugins/Experimental/ActorModifierCore/Source/ActorModifierCore/Private/Subsystems/ActorModifierCoreSubsystem.cpp b/Engine/Plugins/Experimental/ActorModifierCore/Source/ActorModifierCore/Private/Subsystems/ActorModifierCoreSubsystem.cpp
--- a/Engine/Plugins/Experimental/ActorModifierCore/Source/ActorModifierCore/Private/Subsystems/ActorModifierCoreSubsystem.cpp
+++ b/Engine/Plugins/Experimental/ActorModifierCore/Source/ActorModifierCore/Private/Subsystems/ActorModifierCoreSubsystem.cpp
@@ -1,6 +1,9 @@
 // Copyright Epic Games, Inc. All Rights Reserved.
 
 #include "Subsystems/ActorModifierCoreSubsystem.h"
+
+#include "AssetRegistry/AssetData.h"
+#include "AssetRegistry/AssetRegistryModule.h"
 #include "Engine/Engine.h"
 #include "EngineUtils.h"
 #include "GameFramework/Actor.h"
@@ -8,6 +11,7 @@
 #include "Modifiers/ActorModifierCoreComponent.h"
 #include "Modifiers/ActorModifierCoreSharedActor.h"
 #include "Modifiers/ActorModifierCoreStack.h"
+#include "Modifiers/Blueprints/ActorModifierCoreBlueprintBase.h"
 #include "UObject/Class.h"
 #include "UObject/UObjectIterator.h"
 
@@ -24,6 +28,7 @@ UActorModifierCoreSubsystem::FOnModifierClassRegistered UActorModifierCoreSubsys
 UActorModifierCoreSubsystem::FOnModifierClassRegistered UActorModifierCoreSubsystem::OnModifierClassUnregisteredDelegate;
 UActorModifierCoreSubsystem::FOnModifierStackRegistered UActorModifierCoreSubsystem::OnModifierStackRegisteredDelegate;
 UActorModifierCoreSubsystem::FOnModifierStackRegistered UActorModifierCoreSubsystem::OnModifierStackUnregisteredDelegate;
+UActorModifierCoreSubsystem::FOnModifierReplaced UActorModifierCoreSubsystem::OnModifierReplacedDelegate;
 
 UActorModifierCoreSubsystem::UActorModifierCoreSubsystem()
 	: UEngineSubsystem()
@@ -44,6 +49,18 @@ void UActorModifierCoreSubsystem::Initialize(FSubsystemCollectionBase& Collectio
 	Super::Initialize(Collection);
 
 	ScanForModifiers();
+
+	const FAssetRegistryModule& AssetRegistryModule = FModuleManager::LoadModuleChecked<FAssetRegistryModule>("AssetRegistry");
+
+	IAssetRegistry& AssetRegistry = AssetRegistryModule.Get();
+	AssetRegistry.OnFilesLoaded().AddUObject(this, &UActorModifierCoreSubsystem::OnAssetRegistryFilesLoaded);
+	AssetRegistry.OnAssetAdded().AddUObject(this, &UActorModifierCoreSubsystem::OnAssetRegistryAssetAdded);
+	AssetRegistry.OnAssetRemoved().AddUObject(this, &UActorModifierCoreSubsystem::OnAssetRegistryAssetRemoved);
+	AssetRegistry.OnAssetUpdatedOnDisk().AddUObject(this, &UActorModifierCoreSubsystem::OnAssetRegistryAssetUpdated);
+
+#if WITH_EDITOR
+	FCoreUObjectDelegates::OnObjectsReplaced.AddUObject(this, &UActorModifierCoreSubsystem::OnBlueprintObjectsReplaced);
+#endif
 }
 
 void UActorModifierCoreSubsystem::Deinitialize()
@@ -56,6 +73,19 @@ void UActorModifierCoreSubsystem::Deinitialize()
 		UnregisterModifierClass(ModifierName);
 	}
 
+	if (const FAssetRegistryModule* AssetRegistryModule = FModuleManager::GetModulePtr<FAssetRegistryModule>("AssetRegistry"))
+	{
+		IAssetRegistry& AssetRegistry = AssetRegistryModule->Get();
+		AssetRegistry.OnFilesLoaded().RemoveAll(this);
+		AssetRegistry.OnAssetAdded().RemoveAll(this);
+		AssetRegistry.OnAssetRemoved().RemoveAll(this);
+		AssetRegistry.OnInMemoryAssetDeleted().RemoveAll(this);
+	}
+
+#if WITH_EDITOR
+	FCoreUObjectDelegates::OnObjectsReplaced.RemoveAll(this);
+#endif
+
 	Super::Deinitialize();
 }
 
@@ -134,6 +164,19 @@ bool UActorModifierCoreSubsystem::UnregisterModifierClass(const FName& InName)
 	return false;
 }
 
+bool UActorModifierCoreSubsystem::UnregisterModifierClass(const UClass* InModifierClass)
+{
+	for (const TPair<FName, TSharedRef<FActorModifierCoreMetadata>>& ModifierMetadataPair : ModifiersMetadata)
+	{
+		if (ModifierMetadataPair.Value->GetClass() == InModifierClass)
+		{
+			return UnregisterModifierClass(ModifierMetadataPair.Key);
+		}
+	}
+
+	return false;
+}
+
 bool UActorModifierCoreSubsystem::IsRegisteredModifierClass(const FName& InName) const
 {
 	return ModifiersMetadata.Contains(InName);
@@ -1426,4 +1469,75 @@ void UActorModifierCoreSubsystem::ScanForModifiers()
 	}
 }
 
+void UActorModifierCoreSubsystem::OnAssetRegistryFilesLoaded()
+{
+	const FAssetRegistryModule& AssetRegistryModule = FModuleManager::LoadModuleChecked<FAssetRegistryModule>("AssetRegistry");
+
+	TArray<FAssetData> Assets;
+	const FTopLevelAssetPath BlueprintPath(TEXT("/Script/ActorModifierCoreBlueprint.ActorModifierCoreBlueprint"));
+	AssetRegistryModule.Get().GetAssetsByClass(BlueprintPath, Assets, /** Subclass */true);
+
+	for (const FAssetData& Asset : Assets)
+	{
+		RegisterModifierAsset(Asset);
+	}
+}
+
+void UActorModifierCoreSubsystem::OnAssetRegistryAssetAdded(const FAssetData& InAssetData)
+{
+	RegisterModifierAsset(InAssetData);
+}
+
+void UActorModifierCoreSubsystem::OnAssetRegistryAssetUpdated(const FAssetData& InAssetData)
+{
+	RegisterModifierAsset(InAssetData);
+}
+
+void UActorModifierCoreSubsystem::OnAssetRegistryAssetRemoved(const FAssetData& InAssetData)
+{
+	UnregisterModifierAsset(InAssetData);
+}
+
+void UActorModifierCoreSubsystem::OnBlueprintObjectsReplaced(const TMap<UObject*, UObject*>& InReplacements)
+{
+	for (const TPair<UObject*, UObject*>& Replacement : InReplacements)
+	{
+		if (!Replacement.Key || !Replacement.Value || Replacement.Value->IsTemplate())
+		{
+			continue;
+		}
+
+		UActorModifierCoreBlueprintBase* ReplacedModifier = Cast<UActorModifierCoreBlueprintBase>(Replacement.Key);
+		UActorModifierCoreBlueprintBase* ReplacementModifier = Cast<UActorModifierCoreBlueprintBase>(Replacement.Value);
+
+		if (ReplacedModifier && ReplacementModifier)
+		{
+			UnregisterModifierClass(ReplacedModifier->GetClass());
+			RegisterModifierClass(ReplacementModifier->GetClass(), true);
+
+			OnModifierReplacedDelegate.Broadcast(ReplacedModifier, ReplacementModifier);
+		}
+	}
+}
+
+void UActorModifierCoreSubsystem::RegisterModifierAsset(const FAssetData& InAssetData)
+{
+	FString GeneratedClassPath;
+	if (InAssetData.GetTagValue(TEXT("GeneratedClass"), GeneratedClassPath))
+	{
+		const UClass* GeneratedClass = LoadObject<UClass>(nullptr, *GeneratedClassPath);
+		RegisterModifierClass(GeneratedClass);
+	}
+}
+
+void UActorModifierCoreSubsystem::UnregisterModifierAsset(const FAssetData& InAssetData)
+{
+	FString GeneratedClassPath;
+	if (InAssetData.GetTagValue(TEXT("GeneratedClass"), GeneratedClassPath))
+	{
+		const UClass* GeneratedClass = LoadObject<UClass>(nullptr, *GeneratedClassPath);
+		UnregisterModifierClass(GeneratedClass);
+	}
+}
+
 #undef LOCTEXT_NAMESPACE
diff --git a/Engine/Plugins/Experimental/ActorModifierCore/Source/ActorModifierCore/Public/Modifiers/ActorModifierCoreDefs.h b/Engine/Plugins/Experimental/ActorModifierCore/Source/ActorModifierCore/Public/Modifiers/ActorModifierCoreDefs.h
--- a/Engine/Plugins/Experimental/ActorModifierCore/Source/ActorModifierCore/Public/Modifiers/ActorModifierCoreDefs.h
+++ b/Engine/Plugins/Experimental/ActorModifierCore/Source/ActorModifierCore/Public/Modifiers/ActorModifierCoreDefs.h
@@ -27,6 +27,7 @@ class FText;
 class UActorModifierCoreBase;
 class UActorModifierCoreStack;
 
+UENUM(BlueprintType)
 enum class EActorModifierCoreEnableReason : uint8
 {
 	/** Modifier added by user */
@@ -39,6 +40,7 @@ enum class EActorModifierCoreEnableReason : uint8
 	Duplicate
 };
 
+UENUM(BlueprintType)
 enum class EActorModifierCoreDisableReason : uint8
 {
 	/** Modifier disabled by user */
@@ -46,7 +48,7 @@ enum class EActorModifierCoreDisableReason : uint8
 	/** Modifier disabled by undo */
 	Undo,
 	/** Modifier disabled by actor destroyed */
-	Destroyed,
+	Destroyed
 };
 
 /** Enumerates valid positions for modifier operations */
@@ -86,9 +88,14 @@ enum class EActorModifierCoreStatus : uint8
 	Error
 };
 
+DECLARE_DYNAMIC_DELEGATE_RetVal_OneParam(bool, FModifierCompatibilityRule, const AActor*, InTargetActor);
+
 /** Metadata for each modifier CDO, modifier instance will share same metadata as CDO */
+USTRUCT(BlueprintType)
 struct FActorModifierCoreMetadata
 {
+	GENERATED_BODY()
+
 	static inline const FName DefaultCategory = TEXT("Default");
 #if WITH_EDITOR
 	static inline const FText DefaultDescription = FText::FromString(TEXT("Description not provided"));
@@ -142,6 +149,9 @@ struct FActorModifierCoreMetadata
 	/** Sets the usage rule for this modifier, if it passes it will be available for this actor */
 	ACTORMODIFIERCORE_API FActorModifierCoreMetadata& SetCompatibilityRule(const TFunction<bool(const AActor*)>& InModifierRule);
 
+	/** Sets the usage rule for this modifier, if it passes it will be available for this actor */
+	FActorModifierCoreMetadata& SetCompatibilityRule(const FModifierCompatibilityRule& InModifierRule);
+
 	/** Create the modifier instance */
 	UActorModifierCoreBase* CreateModifierInstance(UActorModifierCoreStack* InStack) const;
 
@@ -287,6 +297,7 @@ private:
 
 	/** Rule to pass before this modifier can be used on an actor */
 	TFunction<bool(const AActor*)> CompatibilityRuleFunction;
+	FModifierCompatibilityRule CompatibilityRuleDelegate;
 
 	/** In order to create new instances of this modifier */
 	TSubclassOf<UActorModifierCoreBase> Class;
diff --git a/Engine/Plugins/Experimental/ActorModifierCore/Source/ActorModifierCore/Public/Modifiers/ActorModifierCoreStack.h b/Engine/Plugins/Experimental/ActorModifierCore/Source/ActorModifierCore/Public/Modifiers/ActorModifierCoreStack.h
--- a/Engine/Plugins/Experimental/ActorModifierCore/Source/ActorModifierCore/Public/Modifiers/ActorModifierCoreStack.h
+++ b/Engine/Plugins/Experimental/ActorModifierCore/Source/ActorModifierCore/Public/Modifiers/ActorModifierCoreStack.h
@@ -7,6 +7,7 @@
 #include "ActorModifierCoreDefs.h"
 #include "ActorModifierCoreStack.generated.h"
 
+class UActorModifierCoreBlueprintBase;
 class UActorModifierCoreComponent;
 class USceneComponent;
 
@@ -25,14 +26,25 @@ class UActorModifierCoreStack : public UActorModifierCoreBase
 public:
 	DECLARE_MULTICAST_DELEGATE_OneParam(FOnModifierUpdated, UActorModifierCoreBase* /** UpdatedItem */)
 
-	/** Called when a modifier is added to the stack */
-	ACTORMODIFIERCORE_API static FOnModifierUpdated OnModifierAddedDelegate;
+	static FOnModifierUpdated::RegistrationType& OnModifierAdded()
+	{
+		return OnModifierAddedDelegate;
+	}
 
-	/** Called when a modifier is removed from the stack */
-	ACTORMODIFIERCORE_API static FOnModifierUpdated OnModifierRemovedDelegate;
+	static FOnModifierUpdated::RegistrationType& OnModifierRemoved()
+	{
+		return OnModifierRemovedDelegate;
+	}
 
-	/** Called when a modifier is moved in the stack */
-	ACTORMODIFIERCORE_API static FOnModifierUpdated OnModifierMovedDelegate;
+	static FOnModifierUpdated::RegistrationType& OnModifierMoved()
+	{
+		return OnModifierMovedDelegate;
+	}
+
+	static FOnModifierUpdated::RegistrationType& OnModifierReplaced()
+	{
+		return OnModifierReplacedDelegate;
+	}
 
 	/** Create a new stack by passing the actor and the parent stack if there is one */
 	static UActorModifierCoreStack* Create(UActorModifierCoreComponent* InComponent, UActorModifierCoreStack* InParentStack = nullptr);
@@ -165,6 +177,18 @@ protected:
 	virtual void OnModifierDirty(UActorModifierCoreBase* DirtyModifier, bool bExecute) override;
 
 private:
+	/** Called when a modifier is added to the stack */
+	ACTORMODIFIERCORE_API static FOnModifierUpdated OnModifierAddedDelegate;
+
+	/** Called when a modifier is removed from the stack */
+	ACTORMODIFIERCORE_API static FOnModifierUpdated OnModifierRemovedDelegate;
+
+	/** Called when a modifier is moved in the stack */
+	ACTORMODIFIERCORE_API static FOnModifierUpdated OnModifierMovedDelegate;
+
+	/** Called when a modifier is replaced in the stack (blueprint) */
+	ACTORMODIFIERCORE_API static FOnModifierUpdated OnModifierReplacedDelegate;
+
 	/** Sets the stack to receive tick events */
 	virtual void OnModifierCDOSetup(FActorModifierCoreMetadata& InMetadata) override;
 
@@ -190,6 +214,9 @@ private:
 	/** Checks for any possible modifier optimization within the stack */
 	void CheckModifierOptimization(bool bInInvalidateAll);
 
+	/** Replaces blueprint modifier by their new object */
+	void OnBlueprintModifierReplaced(UActorModifierCoreBlueprintBase* InOldModifier, UActorModifierCoreBlueprintBase* InNewModifier);
+
 	/** Contains actual modifiers in the stack */
 	UPROPERTY(VisibleInstanceOnly, NoClear, Export, Instanced, Category = "Modifiers")
 	TArray<TObjectPtr<UActorModifierCoreBase>> Modifiers;
diff --git a/Engine/Plugins/Experimental/ActorModifierCore/Source/ActorModifierCore/Public/Modifiers/Blueprints/ActorModifierCoreBlueprintBase.h b/Engine/Plugins/Experimental/ActorModifierCore/Source/ActorModifierCore/Public/Modifiers/Blueprints/ActorModifierCoreBlueprintBase.h
new file mode 100644
--- /dev/null
+++ b/Engine/Plugins/Experimental/ActorModifierCore/Source/ActorModifierCore/Public/Modifiers/Blueprints/ActorModifierCoreBlueprintBase.h
@@ -0,0 +1,75 @@
+// Copyright Epic Games, Inc. All Rights Reserved.
+
+#pragma once
+
+#include "Modifiers/ActorModifierCoreBase.h"
+
+#include "ActorModifierCoreBlueprintBase.generated.h"
+
+/** Abstract base class for all blueprint modifier */
+UCLASS(MinimalAPI, Blueprintable, Abstract, EditInlineNew, HideCategories=(Tags, AssetUserData, Activation, Collision, Cooking))
+class UActorModifierCoreBlueprintBase : public UActorModifierCoreBase
+{
+	GENERATED_BODY()
+
+	friend class UActorModifierCoreStack;
+
+protected:
+	/** Called once to setup modifier metadata */
+	UFUNCTION(BlueprintImplementableEvent, Category="Modifiers", meta=(ForceAsFunction))
+	void OnModifierSetupEvent(UPARAM(Ref) FActorModifierCoreMetadata& InMetadata, FActorModifierCoreMetadata& OutMetadata);
+
+	/** Called when the modifier gets recompiled and replaced in the stack */
+	UFUNCTION(BlueprintImplementableEvent, Category="Modifiers", meta=(ForceAsFunction))
+	void OnModifierReplacedEvent(AActor* InTargetActor);
+
+	/** Called when this modifier is added in a stack on an actor */
+	UFUNCTION(BlueprintImplementableEvent, Category="Modifiers", meta=(ForceAsFunction))
+	void OnModifierAddedEvent(AActor* InTargetActor, EActorModifierCoreEnableReason InReason);
+
+	/** Called when this modifier is enabled */
+	UFUNCTION(BlueprintImplementableEvent, Category="Modifiers", meta=(ForceAsFunction))
+	void OnModifierEnabledEvent(AActor* InTargetActor, EActorModifierCoreEnableReason InReason);
+
+	/** Called when this modifier is disabled */
+	UFUNCTION(BlueprintImplementableEvent, Category="Modifiers", meta=(ForceAsFunction))
+	void OnModifierDisabledEvent(AActor* InTargetActor, EActorModifierCoreDisableReason InReason);
+
+	/** Called when this modifier is removed from a stack on an actor */
+	UFUNCTION(BlueprintImplementableEvent, Category="Modifiers", meta=(ForceAsFunction))
+	void OnModifierRemovedEvent(AActor* InTargetActor, EActorModifierCoreDisableReason InReason);
+
+	/** Called before this modifier is applied on an actor to save all relevant state */
+	UFUNCTION(BlueprintImplementableEvent, Category="Modifiers", meta=(ForceAsFunction))
+	void OnModifierSaveStateEvent(AActor* InTargetActor);
+
+	/** Called to restore this modifier actions on an actor */
+	UFUNCTION(BlueprintImplementableEvent, Category="Modifiers", meta=(ForceAsFunction))
+	void OnModifierRestoreStateEvent(AActor* InTargetActor);
+
+	/** Called to apply a custom action on an actor */
+	UFUNCTION(BlueprintImplementableEvent, Category="Modifiers")
+	bool OnModifierApplyEvent(AActor* InTargetActor, FText& OutFailReason);
+
+	/** Flag this modifier as needing an update after a property value has changed */
+	UFUNCTION(BlueprintCallable, Category="Modifiers")
+	void FlagModifierDirty();
+
+private:
+	//~ Begin UObject
+#if WITH_EDITOR
+	virtual void PostEditChangeProperty(FPropertyChangedEvent& InPropertyChangedEvent) override;
+#endif
+	//~ End UObject
+
+	//~ Begin UActorModifierCoreBase
+	virtual void OnModifierCDOSetup(FActorModifierCoreMetadata& InMetadata) override;
+	virtual void OnModifierAdded(EActorModifierCoreEnableReason InReason) override;
+	virtual void OnModifierEnabled(EActorModifierCoreEnableReason InReason) override;
+	virtual void OnModifierDisabled(EActorModifierCoreDisableReason InReason) override;
+	virtual void OnModifierRemoved(EActorModifierCoreDisableReason InReason) override;
+	virtual void SavePreState() override;
+	virtual void RestorePreState() override;
+	virtual void Apply() override;
+	//~ End UActorModifierCoreBase
+};
\ No newline at end of file
diff --git a/Engine/Plugins/Experimental/ActorModifierCore/Source/ActorModifierCore/Public/Modifiers/Blueprints/ActorModifierCoreGeneratedClass.h b/Engine/Plugins/Experimental/ActorModifierCore/Source/ActorModifierCore/Public/Modifiers/Blueprints/ActorModifierCoreGeneratedClass.h
new file mode 100644
--- /dev/null
+++ b/Engine/Plugins/Experimental/ActorModifierCore/Source/ActorModifierCore/Public/Modifiers/Blueprints/ActorModifierCoreGeneratedClass.h
@@ -0,0 +1,14 @@
+// Copyright Epic Games, Inc. All Rights Reserved.
+
+#pragma once
+
+#include "Engine/BlueprintGeneratedClass.h"
+#include "UObject/Package.h"
+#include "ActorModifierCoreGeneratedClass.generated.h"
+
+/** Blueprint modifier class generated by UActorModifierCoreBlueprint */
+UCLASS(MinimalAPI)
+class UActorModifierCoreGeneratedClass : public UBlueprintGeneratedClass
+{
+	GENERATED_BODY()
+};
diff --git a/Engine/Plugins/Experimental/ActorModifierCore/Source/ActorModifierCore/Public/Modifiers/Utilities/ActorModifierCoreLibrary.h b/Engine/Plugins/Experimental/ActorModifierCore/Source/ActorModifierCore/Public/Modifiers/Utilities/ActorModifierCoreLibrary.h
--- a/Engine/Plugins/Experimental/ActorModifierCore/Source/ActorModifierCore/Public/Modifiers/Utilities/ActorModifierCoreLibrary.h
+++ b/Engine/Plugins/Experimental/ActorModifierCore/Source/ActorModifierCore/Public/Modifiers/Utilities/ActorModifierCoreLibrary.h
@@ -226,4 +226,66 @@ public:
 	 */
 	UFUNCTION(BlueprintCallable, Category="Motion Design|Modifiers|Utility")
 	static ACTORMODIFIERCORE_API bool GetAvailableModifiers(TSet<TSubclassOf<UActorModifierCoreBase>>& OutAvailableModifierClasses);
+
+	/**
+	 * Sets the modifier metadata name
+	 * @param InMetadata The modifier metadata to use
+	 * @param InName The modifier name to set
+	 * @return The modifier metadata to chain operations
+	 */
+	UFUNCTION(BlueprintCallable, Category="Motion Design|Modifiers|Utility")
+	static ACTORMODIFIERCORE_API FActorModifierCoreMetadata& SetModifierMetadataName(UPARAM(Ref) FActorModifierCoreMetadata& InMetadata, FName InName);
+
+	/**
+	 * Sets the modifier metadata category
+	 * @param InMetadata The modifier metadata to use
+	 * @param InCategory The modifier category to set
+	 * @return The modifier metadata to chain operations
+	 */
+	UFUNCTION(BlueprintCallable, Category="Motion Design|Modifiers|Utility")
+	static ACTORMODIFIERCORE_API FActorModifierCoreMetadata& SetModifierMetadataCategory(UPARAM(Ref) FActorModifierCoreMetadata& InMetadata, FName InCategory);
+
+	/**
+	 * Sets the modifier metadata display name (EDITOR-ONLY)
+	 * @param InMetadata The modifier metadata to use
+	 * @param InName The modifier name to set
+	 * @return The modifier metadata to chain operations
+	 */
+	UFUNCTION(BlueprintCallable, Category="Motion Design|Modifiers|Utility")
+	static ACTORMODIFIERCORE_API FActorModifierCoreMetadata& SetModifierMetadataDisplayName(UPARAM(Ref) FActorModifierCoreMetadata& InMetadata, const FText& InName);
+
+	/**
+	 * Sets the modifier metadata color (EDITOR-ONLY)
+	 * @param InMetadata The modifier metadata to use
+	 * @param InColor The modifier color to set
+	 * @return The modifier metadata to chain operations
+	 */
+	static ACTORMODIFIERCORE_API FActorModifierCoreMetadata& SetModifierMetadataColor(UPARAM(Ref) FActorModifierCoreMetadata& InMetadata, const FLinearColor& InColor);
+
+	/**
+	 * Sets the modifier metadata description (EDITOR-ONLY)
+	 * @param InMetadata The modifier metadata to use
+	 * @param InDescription The modifier description to set
+	 * @return The modifier metadata to chain operations
+	 */
+	UFUNCTION(BlueprintCallable, Category="Motion Design|Modifiers|Utility")
+	static ACTORMODIFIERCORE_API FActorModifierCoreMetadata& SetModifierMetadataDescription(UPARAM(Ref) FActorModifierCoreMetadata& InMetadata, const FText& InDescription);
+
+	/**
+	 * Adds a modifier metadata dependency for this modifier
+	 * @param InMetadata The modifier metadata to use
+	 * @param InModifierClass The modifier dependency to add
+	 * @return The modifier metadata to chain operations
+	 */
+	UFUNCTION(BlueprintCallable, Category="Motion Design|Modifiers|Utility")
+	static ACTORMODIFIERCORE_API FActorModifierCoreMetadata& AddModifierMetadataDependency(UPARAM(Ref) FActorModifierCoreMetadata& InMetadata, TSubclassOf<UActorModifierCoreBase> InModifierClass);
+
+	/**
+	 * Sets the modifier metadata compatibility rule
+	 * @param InMetadata The modifier metadata to use
+	 * @param InDelegate The modifier rule to set
+	 * @return The modifier metadata to chain operations
+	 */
+	UFUNCTION(BlueprintCallable, Category="Motion Design|Modifiers|Utility")
+	static ACTORMODIFIERCORE_API FActorModifierCoreMetadata& SetModifierMetadataCompatibilityRule(UPARAM(Ref) FActorModifierCoreMetadata& InMetadata, const FModifierCompatibilityRule& InDelegate);
 };
\ No newline at end of file
diff --git a/Engine/Plugins/Experimental/ActorModifierCore/Source/ActorModifierCore/Public/Subsystems/ActorModifierCoreSubsystem.h b/Engine/Plugins/Experimental/ActorModifierCore/Source/ActorModifierCore/Public/Subsystems/ActorModifierCoreSubsystem.h
--- a/Engine/Plugins/Experimental/ActorModifierCore/Source/ActorModifierCore/Public/Subsystems/ActorModifierCoreSubsystem.h
+++ b/Engine/Plugins/Experimental/ActorModifierCore/Source/ActorModifierCore/Public/Subsystems/ActorModifierCoreSubsystem.h
@@ -10,8 +10,10 @@ class AActor;
 class AActorModifierCoreSharedActor;
 class UClass;
 class UActorModifierCoreBase;
+class UActorModifierCoreBlueprintBase;
 class UActorModifierCoreSharedObject;
 class UActorModifierCoreStack;
+struct FAssetData;
 
 /** This subsystem handle all modifiers stack active in the engine and allows to create modifiers with registered metadata */
 UCLASS()
@@ -25,29 +27,33 @@ class UActorModifierCoreSubsystem : public UEngineSubsystem
 	friend struct FActorModifierCoreMetadata;
 
 public:
-	DECLARE_MULTICAST_DELEGATE_OneParam(FOnModifierClassRegistered, const FActorModifierCoreMetadata& /* ModifierMetadata */)
-
-	/** Called when a modifier class is registered */
-	ACTORMODIFIERCORE_API static FOnModifierClassRegistered OnModifierClassRegisteredDelegate;
-
-	/** Called when a modifier class is unregistered */
-	ACTORMODIFIERCORE_API static FOnModifierClassRegistered OnModifierClassUnregisteredDelegate;
-
+	DECLARE_MULTICAST_DELEGATE_OneParam(FOnModifierClassRegistered, const FActorModifierCoreMetadata& /** ModifierMetadata */)
 	DECLARE_MULTICAST_DELEGATE_OneParam(FOnModifierStackRegistered, const UActorModifierCoreStack* /** ActorRootStack */)
+	DECLARE_MULTICAST_DELEGATE_TwoParams(FOnModifierReplaced, UActorModifierCoreBlueprintBase* /** PreviousModifier */, UActorModifierCoreBlueprintBase* /** ReplacementModifier */)
 
-	/** Called when an actor modifier stack is registered */
-	static FOnModifierStackRegistered OnModifierStackRegisteredDelegate;
+	static FOnModifierClassRegistered::RegistrationType& OnModifierClassRegistered()
+	{
+		return OnModifierClassRegisteredDelegate;
+	}
 
-	/** Called when an actor modifier stack is unregistered */
-	static FOnModifierStackRegistered OnModifierStackUnregisteredDelegate;
+	static FOnModifierClassRegistered::RegistrationType& OnModifierClassUnregistered()
+	{
+		return OnModifierClassUnregisteredDelegate;
+	}
 
-	UActorModifierCoreSubsystem();
+	static FOnModifierReplaced::RegistrationType& OnModifierReplaced()
+	{
+		return OnModifierReplacedDelegate;
+	}
 
 	ACTORMODIFIERCORE_API static UActorModifierCoreSubsystem* Get();
 
+	UActorModifierCoreSubsystem();
+
 	/** Register a modifier class or override an already existing one */
 	ACTORMODIFIERCORE_API bool RegisterModifierClass(const UClass* InModifierClass, bool bInOverrideIfExists = false);
 	ACTORMODIFIERCORE_API bool UnregisterModifierClass(const FName& InName);
+	ACTORMODIFIERCORE_API bool UnregisterModifierClass(const UClass* InModifierClass);
 	ACTORMODIFIERCORE_API bool IsRegisteredModifierClass(const FName& InName) const;
 	ACTORMODIFIERCORE_API bool IsRegisteredModifierClass(const UClass* InClass) const;
 
@@ -129,6 +135,21 @@ public:
 	ACTORMODIFIERCORE_API void GetSortedModifiers(const TSet<UActorModifierCoreBase*>& InModifiers, AActor* InTargetActor, UActorModifierCoreBase* InTargetModifier, EActorModifierCoreStackPosition InPosition, TArray<UActorModifierCoreBase*>& OutMoveModifiers, TArray<UActorModifierCoreBase*>& OutCloneModifiers) const;
 
 protected:
+	/** Called when a modifier class is registered */
+	ACTORMODIFIERCORE_API static FOnModifierClassRegistered OnModifierClassRegisteredDelegate;
+
+	/** Called when a modifier class is unregistered */
+	ACTORMODIFIERCORE_API static FOnModifierClassRegistered OnModifierClassUnregisteredDelegate;
+
+	/** Called when an actor modifier stack is registered */
+	static FOnModifierStackRegistered OnModifierStackRegisteredDelegate;
+
+	/** Called when an actor modifier stack is unregistered */
+	static FOnModifierStackRegistered OnModifierStackUnregisteredDelegate;
+
+	/** Called when a blueprint modifier gets replaced */
+	static FOnModifierReplaced OnModifierReplacedDelegate;
+
 	//~ Begin UEngineSubsystem
 	virtual void Initialize(FSubsystemCollectionBase& Collection) override;
 	virtual void Deinitialize() override;
@@ -163,6 +184,16 @@ protected:
 	/** Done only once at subsystem initialization */
 	void ScanForModifiers();
 
+	void OnAssetRegistryFilesLoaded();
+	void OnAssetRegistryAssetAdded(const FAssetData& InAssetData);
+	void OnAssetRegistryAssetUpdated(const FAssetData& InAssetData);
+	void OnAssetRegistryAssetRemoved(const FAssetData& InAssetData);
+
+	void OnBlueprintObjectsReplaced(const TMap<UObject*, UObject*>& InReplacements);
+
+	void RegisterModifierAsset(const FAssetData& InAssetData);
+	void UnregisterModifierAsset(const FAssetData& InAssetData);
+
 	/** Registers a modifier shared provider actor for a world */
 	bool RegisterModifierSharedProvider(AActorModifierCoreSharedActor* InSharedActor) const;
 
diff --git a/Engine/Plugins/Experimental/ActorModifierCore/Source/ActorModifierCoreBlueprint/ActorModifierCoreBlueprint.Build.cs b/Engine/Plugins/Experimental/ActorModifierCore/Source/ActorModifierCoreBlueprint/ActorModifierCoreBlueprint.Build.cs
new file mode 100644
--- /dev/null
+++ b/Engine/Plugins/Experimental/ActorModifierCore/Source/ActorModifierCoreBlueprint/ActorModifierCoreBlueprint.Build.cs
@@ -0,0 +1,29 @@
+﻿// Copyright Epic Games, Inc. All Rights Reserved.
+
+using UnrealBuildTool;
+
+public class ActorModifierCoreBlueprint : ModuleRules
+{
+    public ActorModifierCoreBlueprint(ReadOnlyTargetRules Target) : base(Target)
+    {
+        PCHUsage = ModuleRules.PCHUsageMode.UseExplicitOrSharedPCHs;
+        
+        PublicDependencyModuleNames.AddRange(
+	        new string[]
+	        {
+		        "BlueprintGraph",
+		        "Core",
+		        "CoreUObject",
+		        "Engine",
+		        "SlateCore"
+	        }
+        );
+        
+        PrivateDependencyModuleNames.AddRange(
+	        new string[]
+	        {
+		        "ActorModifierCore"
+	        }
+        );
+    }
+}
\ No newline at end of file
diff --git a/Engine/Plugins/Experimental/ActorModifierCore/Source/ActorModifierCoreBlueprint/Private/ActorModifierCoreBlueprint.cpp b/Engine/Plugins/Experimental/ActorModifierCore/Source/ActorModifierCoreBlueprint/Private/ActorModifierCoreBlueprint.cpp
new file mode 100644
--- /dev/null
+++ b/Engine/Plugins/Experimental/ActorModifierCore/Source/ActorModifierCoreBlueprint/Private/ActorModifierCoreBlueprint.cpp
@@ -0,0 +1,22 @@
+// Copyright Epic Games, Inc. All Rights Reserved.
+
+#include "ActorModifierCoreBlueprint.h"
+
+#include "Modifiers/Blueprints/ActorModifierCoreGeneratedClass.h"
+
+#if WITH_EDITOR
+UClass* UActorModifierCoreBlueprint::GetBlueprintClass() const
+{
+	return UActorModifierCoreGeneratedClass::StaticClass();
+}
+
+void UActorModifierCoreBlueprint::GetReparentingRules(TSet<const UClass*>& OutAllowedChildrenOfClasses, TSet<const UClass*>& OutDisallowedChildrenOfClasses) const
+{
+	OutAllowedChildrenOfClasses.Add(UActorModifierCoreGeneratedClass::StaticClass());
+}
+
+bool UActorModifierCoreBlueprint::SupportedByDefaultBlueprintFactory() const
+{
+	return false;
+}
+#endif
diff --git a/Engine/Plugins/Experimental/ActorModifierCore/Source/ActorModifierCoreBlueprint/Private/ActorModifierCoreBlueprintModule.cpp b/Engine/Plugins/Experimental/ActorModifierCore/Source/ActorModifierCoreBlueprint/Private/ActorModifierCoreBlueprintModule.cpp
new file mode 100644
--- /dev/null
+++ b/Engine/Plugins/Experimental/ActorModifierCore/Source/ActorModifierCoreBlueprint/Private/ActorModifierCoreBlueprintModule.cpp
@@ -0,0 +1,5 @@
+﻿// Copyright Epic Games, Inc. All Rights Reserved.
+
+#include "Modules/ModuleManager.h"
+
+IMPLEMENT_MODULE(FDefaultModuleImpl, ActorModifierCoreBlueprint)
\ No newline at end of file
diff --git a/Engine/Plugins/Experimental/ActorModifierCore/Source/ActorModifierCoreBlueprint/Public/ActorModifierCoreBlueprint.h b/Engine/Plugins/Experimental/ActorModifierCore/Source/ActorModifierCoreBlueprint/Public/ActorModifierCoreBlueprint.h
new file mode 100644
--- /dev/null
+++ b/Engine/Plugins/Experimental/ActorModifierCore/Source/ActorModifierCoreBlueprint/Public/ActorModifierCoreBlueprint.h
@@ -0,0 +1,23 @@
+﻿// Copyright Epic Games, Inc. All Rights Reserved.
+
+#pragma once
+
+#include "Engine/Blueprint.h"
+#include "ActorModifierCoreBlueprint.generated.h"
+
+/**
+ * Blueprint object that defines a user modifier
+ */
+UCLASS(MinimalAPI)
+class UActorModifierCoreBlueprint : public UBlueprint
+{
+	GENERATED_BODY()
+
+	//~ Begin UBlueprint
+#if WITH_EDITOR
+	virtual UClass* GetBlueprintClass() const override;
+	virtual void GetReparentingRules(TSet<const UClass*>& OutAllowedChildrenOfClasses, TSet<const UClass*>& OutDisallowedChildrenOfClasses) const override;
+	virtual bool SupportedByDefaultBlueprintFactory() const override;
+#endif
+	//~ End UBlueprint
+};
diff --git a/Engine/Plugins/Experimental/ActorModifierCore/Source/ActorModifierCoreEditor/ActorModifierCoreEditor.Build.cs b/Engine/Plugins/Experimental/ActorModifierCore/Source/ActorModifierCoreEditor/ActorModifierCoreEditor.Build.cs
--- a/Engine/Plugins/Experimental/ActorModifierCore/Source/ActorModifierCoreEditor/ActorModifierCoreEditor.Build.cs
+++ b/Engine/Plugins/Experimental/ActorModifierCore/Source/ActorModifierCoreEditor/ActorModifierCoreEditor.Build.cs
@@ -20,7 +20,9 @@ public class ActorModifierCoreEditor : ModuleRules
 			new string[]
 			{
 				"ActorModifierCore",
+				"ActorModifierCoreBlueprint",
 				"ApplicationCore",
+				"AssetDefinition",
 				"CoreUObject",
 				"CustomDetailsView",
 				"EditorSubsystem",
@@ -36,7 +38,7 @@ public class ActorModifierCoreEditor : ModuleRules
 				"ToolMenus",
 				"ToolWidgets",
 				"UnrealEd",
-				"WorkspaceMenuStructure",
+				"WorkspaceMenuStructure"
 			}
 		);
 	}
diff --git a/Engine/Plugins/Experimental/ActorModifierCore/Source/ActorModifierCoreEditor/Private/ActorModifierCoreEditorModule.cpp b/Engine/Plugins/Experimental/ActorModifierCore/Source/ActorModifierCoreEditor/Private/ActorModifierCoreEditorModule.cpp
--- a/Engine/Plugins/Experimental/ActorModifierCore/Source/ActorModifierCoreEditor/Private/ActorModifierCoreEditorModule.cpp
+++ b/Engine/Plugins/Experimental/ActorModifierCore/Source/ActorModifierCoreEditor/Private/ActorModifierCoreEditorModule.cpp
@@ -2,22 +2,26 @@
 
 #include "ActorModifierCoreEditorModule.h"
 
+#include "ActorModifierCoreBlueprint.h"
+#include "KismetCompilerModule.h"
 #include "Modifiers/ActorModifierCoreStack.h"
+#include "Modifiers/Blueprints/ActorModifierCoreBlueprintBase.h"
 #include "Modifiers/Customizations/ActorModifierCoreEditorDetailCustomization.h"
 #include "Modules/ModuleManager.h"
 #include "PropertyEditorModule.h"
 
 void FActorModifierCoreEditorModule::StartupModule()
 {
-	RegisterCustomizations();
+	RegisterDetailCustomizations();
+	RegisterBlueprintCustomizations();
 }
 
 void FActorModifierCoreEditorModule::ShutdownModule()
 {
-	UnregisterCustomizations();
+	UnregisterDetailCustomizations();
 }
 
-void FActorModifierCoreEditorModule::RegisterCustomizations() const
+void FActorModifierCoreEditorModule::RegisterDetailCustomizations() const
 {
 	// Register custom layouts
 	static FName PropertyEditor("PropertyEditor");
@@ -27,7 +31,7 @@ void FActorModifierCoreEditorModule::RegisterCustomizations() const
 		, FOnGetDetailCustomizationInstance::CreateStatic(&FActorModifierCoreEditorDetailCustomization::MakeInstance));
 }
 
-void FActorModifierCoreEditorModule::UnregisterCustomizations() const
+void FActorModifierCoreEditorModule::UnregisterDetailCustomizations() const
 {
 	// Unregister custom layouts
 	static FName PropertyEditor("PropertyEditor");
@@ -38,4 +42,10 @@ void FActorModifierCoreEditorModule::UnregisterCustomizations() const
 	}
 }
 
+void FActorModifierCoreEditorModule::RegisterBlueprintCustomizations()
+{
+	IKismetCompilerInterface& KismetCompilerModule = FModuleManager::LoadModuleChecked<IKismetCompilerInterface>("KismetCompiler");
+	KismetCompilerModule.OverrideBPTypeForClass(UActorModifierCoreBlueprintBase::StaticClass(), UActorModifierCoreBlueprint::StaticClass());
+}
+
 IMPLEMENT_MODULE(FActorModifierCoreEditorModule, ActorModifierCoreEditor)
diff --git a/Engine/Plugins/Experimental/ActorModifierCore/Source/ActorModifierCoreEditor/Private/ActorModifierCoreEditorModule.h b/Engine/Plugins/Experimental/ActorModifierCore/Source/ActorModifierCoreEditor/Private/ActorModifierCoreEditorModule.h
--- a/Engine/Plugins/Experimental/ActorModifierCore/Source/ActorModifierCoreEditor/Private/ActorModifierCoreEditorModule.h
+++ b/Engine/Plugins/Experimental/ActorModifierCore/Source/ActorModifierCoreEditor/Private/ActorModifierCoreEditorModule.h
@@ -13,6 +13,8 @@ public:
 	//~ End IModuleInterface
 
 protected:
-	void RegisterCustomizations() const;
-	void UnregisterCustomizations() const;
+	void RegisterDetailCustomizations() const;
+	void UnregisterDetailCustomizations() const;
+
+	void RegisterBlueprintCustomizations();
 };
\ No newline at end of file
diff --git a/Engine/Plugins/Experimental/ActorModifierCore/Source/ActorModifierCoreEditor/Private/Modifiers/Blueprints/ActorModifierCoreBlueprintFactory.cpp b/Engine/Plugins/Experimental/ActorModifierCore/Source/ActorModifierCoreEditor/Private/Modifiers/Blueprints/ActorModifierCoreBlueprintFactory.cpp
new file mode 100644
--- /dev/null
+++ b/Engine/Plugins/Experimental/ActorModifierCore/Source/ActorModifierCoreEditor/Private/Modifiers/Blueprints/ActorModifierCoreBlueprintFactory.cpp
@@ -0,0 +1,48 @@
+// Copyright Epic Games, Inc. All Rights Reserved.
+
+#include "ActorModifierCoreBlueprintFactory.h"
+
+#include "ActorModifierCoreBlueprint.h"
+#include "Kismet2/KismetEditorUtilities.h"
+#include "Misc/MessageDialog.h"
+#include "Modifiers/Blueprints/ActorModifierCoreBlueprintBase.h"
+#include "Modifiers/Blueprints/ActorModifierCoreGeneratedClass.h"
+
+#define LOCTEXT_NAMESPACE "ActorModifierCoreBlueprintFactory"
+
+UActorModifierCoreBlueprintFactory::UActorModifierCoreBlueprintFactory()
+{
+	ParentClass = UActorModifierCoreBlueprintBase::StaticClass();
+	SupportedClass = UActorModifierCoreBlueprint::StaticClass();
+	bCreateNew = true;
+	bEditorImport = false;
+	bEditAfterNew = true;
+}
+
+UObject* UActorModifierCoreBlueprintFactory::FactoryCreateNew(UClass* InClass, UObject* InParent, FName InName, EObjectFlags InFlags, UObject* InContext, FFeedbackContext* InWarn, FName InCallingContext)
+{
+	check(InClass && InClass->IsChildOf<UActorModifierCoreBlueprint>());
+
+	if (!FKismetEditorUtilities::CanCreateBlueprintOfClass(ParentClass) || !ParentClass->IsChildOf<UActorModifierCoreBlueprintBase>())
+	{
+		FMessageDialog::Open(EAppMsgType::Ok
+			, FText::Format(LOCTEXT("InvalidParentClassMessage", "Unable to create Modifier Blueprint with parent class '{0}'.")
+			, FText::FromString(GetNameSafe(ParentClass))));
+		return nullptr;
+	}
+
+	// Create Blueprint
+	UActorModifierCoreBlueprint* Blueprint = CastChecked<UActorModifierCoreBlueprint>(FKismetEditorUtilities::CreateBlueprint(
+		ParentClass
+		, InParent
+		, InName
+		, EBlueprintType::BPTYPE_Normal
+		, UActorModifierCoreBlueprint::StaticClass()
+		, UActorModifierCoreGeneratedClass::StaticClass()
+		, InCallingContext)
+	);
+
+	return Blueprint;
+}
+
+#undef LOCTEXT_NAMESPACE
diff --git a/Engine/Plugins/Experimental/ActorModifierCore/Source/ActorModifierCoreEditor/Private/Modifiers/Blueprints/ActorModifierCoreBlueprintFactory.h b/Engine/Plugins/Experimental/ActorModifierCore/Source/ActorModifierCoreEditor/Private/Modifiers/Blueprints/ActorModifierCoreBlueprintFactory.h
new file mode 100644
--- /dev/null
+++ b/Engine/Plugins/Experimental/ActorModifierCore/Source/ActorModifierCoreEditor/Private/Modifiers/Blueprints/ActorModifierCoreBlueprintFactory.h
@@ -0,0 +1,26 @@
+// Copyright Epic Games, Inc. All Rights Reserved.
+
+#pragma once
+
+#include "Factories/Factory.h"
+#include "Templates/SubclassOf.h"
+#include "ActorModifierCoreBlueprintFactory.generated.h"
+
+class UActorModifierCoreBlueprintBase;
+
+/** Factory responsible to create a UActorModifierCoreBlueprint object */
+UCLASS()
+class UActorModifierCoreBlueprintFactory : public UFactory
+{
+	GENERATED_BODY()
+
+public:
+	UActorModifierCoreBlueprintFactory();
+
+	//~ Begin UFactory
+	virtual UObject* FactoryCreateNew(UClass* InClass, UObject* InParent, FName InName, EObjectFlags InFlags, UObject* InContext, FFeedbackContext* InWarn, FName InCallingContext) override;
+	//~ End UFactory
+
+	UPROPERTY()
+	TSubclassOf<UActorModifierCoreBlueprintBase> ParentClass;
+};
diff --git a/Engine/Plugins/Experimental/ActorModifierCore/Source/ActorModifierCoreEditor/Private/Modifiers/Blueprints/AssetDefinition_ActorModifierCoreBlueprint.cpp b/Engine/Plugins/Experimental/ActorModifierCore/Source/ActorModifierCoreEditor/Private/Modifiers/Blueprints/AssetDefinition_ActorModifierCoreBlueprint.cpp
new file mode 100644
--- /dev/null
+++ b/Engine/Plugins/Experimental/ActorModifierCore/Source/ActorModifierCoreEditor/Private/Modifiers/Blueprints/AssetDefinition_ActorModifierCoreBlueprint.cpp
@@ -0,0 +1,57 @@
+// Copyright Epic Games, Inc. All Rights Reserved.
+
+#include "AssetDefinition_ActorModifierCoreBlueprint.h"
+
+#include "ActorModifierCoreBlueprint.h"
+#include "BlueprintEditorModule.h"
+#include "Kismet2/BlueprintEditorUtils.h"
+#include "Misc/MessageDialog.h"
+
+#define LOCTEXT_NAMESPACE "AssetDefinition_ActorModifierCoreBlueprint"
+
+FText UAssetDefinition_ActorModifierCoreBlueprint::GetAssetDisplayName() const
+{
+	return LOCTEXT("AssetDisplayName", "Actor Modifier Blueprint");
+}
+
+FLinearColor UAssetDefinition_ActorModifierCoreBlueprint::GetAssetColor() const
+{
+	return FLinearColor::White;
+}
+
+TSoftClassPtr<UObject> UAssetDefinition_ActorModifierCoreBlueprint::GetAssetClass() const
+{
+	return UActorModifierCoreBlueprint::StaticClass();
+}
+
+TConstArrayView<FAssetCategoryPath> UAssetDefinition_ActorModifierCoreBlueprint::GetAssetCategories() const
+{
+	return MakeArrayView(&EAssetCategoryPaths::Blueprint, 1);
+}
+
+EAssetCommandResult UAssetDefinition_ActorModifierCoreBlueprint::OpenAssets(const FAssetOpenArgs& InOpenArgs) const
+{
+	for (UActorModifierCoreBlueprint* Blueprint : InOpenArgs.LoadObjects<UActorModifierCoreBlueprint>())
+	{
+		bool bShouldOpen = true;
+
+		if (!Blueprint->SkeletonGeneratedClass || !Blueprint->GeneratedClass)
+		{
+			bShouldOpen = EAppReturnType::Yes == FMessageDialog::Open(EAppMsgType::YesNo
+				, LOCTEXT("InvalidBlueprintClassPrompt",
+					"Blueprint could not be loaded because it derives from an invalid class.\n"
+					"Check to make sure the parent class for this blueprint hasn't been removed!\n"
+					"Do you want to continue (it can crash the editor)?"));
+		}
+
+		if (bShouldOpen)
+		{
+			FBlueprintEditorModule& BlueprintEditorModule = FModuleManager::LoadModuleChecked<FBlueprintEditorModule>("Kismet");
+			TSharedRef<IBlueprintEditor> NewKismetEditor = BlueprintEditorModule.CreateBlueprintEditor(InOpenArgs.GetToolkitMode(), InOpenArgs.ToolkitHost, Blueprint, FBlueprintEditorUtils::ShouldOpenWithDataOnlyEditor(Blueprint));
+		}
+	}
+
+	return EAssetCommandResult::Handled;
+}
+
+#undef LOCTEXT_NAMESPACE
diff --git a/Engine/Plugins/Experimental/ActorModifierCore/Source/ActorModifierCoreEditor/Private/Modifiers/Blueprints/AssetDefinition_ActorModifierCoreBlueprint.h b/Engine/Plugins/Experimental/ActorModifierCore/Source/ActorModifierCoreEditor/Private/Modifiers/Blueprints/AssetDefinition_ActorModifierCoreBlueprint.h
new file mode 100644
--- /dev/null
+++ b/Engine/Plugins/Experimental/ActorModifierCore/Source/ActorModifierCoreEditor/Private/Modifiers/Blueprints/AssetDefinition_ActorModifierCoreBlueprint.h
@@ -0,0 +1,22 @@
+// Copyright Epic Games, Inc. All Rights Reserved.
+
+#pragma once
+
+#include "AssetDefinition.h"
+#include "AssetDefinition_ActorModifierCoreBlueprint.generated.h"
+
+/** Asset definition for UActorModifierCoreBlueprint objects available in editor */
+UCLASS()
+class UAssetDefinition_ActorModifierCoreBlueprint : public UAssetDefinition
+{
+	GENERATED_BODY()
+
+public:
+	//~ Begin UAssetDefinition
+	virtual FText GetAssetDisplayName() const override;
+	virtual FLinearColor GetAssetColor() const override;
+	virtual TSoftClassPtr<UObject> GetAssetClass() const override;
+	virtual TConstArrayView<FAssetCategoryPath> GetAssetCategories() const override;
+	virtual EAssetCommandResult OpenAssets(const FAssetOpenArgs& InOpenArgs) const override;
+	//~ End UAssetDefinition
+};
diff --git a/Engine/Plugins/Experimental/ActorModifierCore/Source/ActorModifierCoreEditor/Private/Modifiers/Customizations/ActorModifierCoreEditorStackCustomization.cpp b/Engine/Plugins/Experimental/ActorModifierCore/Source/ActorModifierCoreEditor/Private/Modifiers/Customizations/ActorModifierCoreEditorStackCustomization.cpp
--- a/Engine/Plugins/Experimental/ActorModifierCore/Source/ActorModifierCoreEditor/Private/Modifiers/Customizations/ActorModifierCoreEditorStackCustomization.cpp
+++ b/Engine/Plugins/Experimental/ActorModifierCore/Source/ActorModifierCoreEditor/Private/Modifiers/Customizations/ActorModifierCoreEditorStackCustomization.cpp
@@ -103,6 +103,11 @@ bool UActorModifierCoreEditorStackCustomization::GetChildrenItem(const FOperator
 			{
 				for (UActorModifierCoreBase* Modifier : ModifierStack->GetModifiers())
 				{
+					if (!IsValid(Modifier))
+					{
+						continue;
+					}
+
 					if (const int32* Index = ClassToIndex.Find(Modifier->GetClass()))
 					{
 						Modifiers[*Index].Add(MakeShared<FOperatorStackEditorObjectItem>(Modifier));
@@ -127,7 +132,10 @@ bool UActorModifierCoreEditorStackCustomization::GetChildrenItem(const FOperator
 		{
 			for (UActorModifierCoreBase* Modifier : InItem->Get<UActorModifierCoreStack>(0)->GetModifiers())
 			{
-				OutChildrenItems.Add(MakeShared<FOperatorStackEditorObjectItem>(Modifier));
+				if (IsValid(Modifier))
+				{
+					OutChildrenItems.Add(MakeShared<FOperatorStackEditorObjectItem>(Modifier));
+				}
 			}
 		}
 	}
diff --git a/Engine/Plugins/Experimental/Avalanche/Source/AvalancheEditor/Private/OperatorStack/Widgets/SAvaOperatorStackTab.cpp b/Engine/Plugins/Experimental/Avalanche/Source/AvalancheEditor/Private/OperatorStack/Widgets/SAvaOperatorStackTab.cpp
--- a/Engine/Plugins/Experimental/Avalanche/Source/AvalancheEditor/Private/OperatorStack/Widgets/SAvaOperatorStackTab.cpp
+++ b/Engine/Plugins/Experimental/Avalanche/Source/AvalancheEditor/Private/OperatorStack/Widgets/SAvaOperatorStackTab.cpp
@@ -28,9 +28,10 @@ void SAvaOperatorStackTab::Construct(const FArguments& InArgs
 	USelection::SelectionChangedEvent.AddSP(this, &SAvaOperatorStackTab::RefreshSelection);
 
 	// Modifiers delegates
-	UActorModifierCoreStack::OnModifierAddedDelegate.AddSP(this, &SAvaOperatorStackTab::OnModifierUpdated);
-	UActorModifierCoreStack::OnModifierMovedDelegate.AddSP(this, &SAvaOperatorStackTab::OnModifierUpdated);
-	UActorModifierCoreStack::OnModifierRemovedDelegate.AddSP(this, &SAvaOperatorStackTab::OnModifierUpdated);
+	UActorModifierCoreStack::OnModifierAdded().AddSP(this, &SAvaOperatorStackTab::OnModifierUpdated);
+	UActorModifierCoreStack::OnModifierMoved().AddSP(this, &SAvaOperatorStackTab::OnModifierUpdated);
+	UActorModifierCoreStack::OnModifierRemoved().AddSP(this, &SAvaOperatorStackTab::OnModifierUpdated);
+	UActorModifierCoreStack::OnModifierReplaced().AddSP(this, &SAvaOperatorStackTab::OnModifierUpdated);
 
 	// Property controllers delegates
 	UPropertyAnimatorCoreBase::OnAnimatorCreatedDelegate.AddSP(this, &SAvaOperatorStackTab::OnAnimatorUpdated);
@@ -55,9 +56,10 @@ SAvaOperatorStackTab::~SAvaOperatorStackTab()
 {
 	USelection::SelectionChangedEvent.RemoveAll(this);
 
-	UActorModifierCoreStack::OnModifierAddedDelegate.RemoveAll(this);
-	UActorModifierCoreStack::OnModifierMovedDelegate.RemoveAll(this);
-	UActorModifierCoreStack::OnModifierRemovedDelegate.RemoveAll(this);
+	UActorModifierCoreStack::OnModifierAdded().RemoveAll(this);
+	UActorModifierCoreStack::OnModifierMoved().RemoveAll(this);
+	UActorModifierCoreStack::OnModifierRemoved().RemoveAll(this);
+	UActorModifierCoreStack::OnModifierReplaced().RemoveAll(this);
 
 	UPropertyAnimatorCoreBase::OnAnimatorCreatedDelegate.RemoveAll(this);
 	UPropertyAnimatorCoreBase::OnAnimatorRemovedDelegate.RemoveAll(this);
diff --git a/Engine/Plugins/Experimental/Avalanche/Source/AvalancheModifiersEditor/Private/AvaModifiersEditorStyle.cpp b/Engine/Plugins/Experimental/Avalanche/Source/AvalancheModifiersEditor/Private/AvaModifiersEditorStyle.cpp
--- a/Engine/Plugins/Experimental/Avalanche/Source/AvalancheModifiersEditor/Private/AvaModifiersEditorStyle.cpp
+++ b/Engine/Plugins/Experimental/Avalanche/Source/AvalancheModifiersEditor/Private/AvaModifiersEditorStyle.cpp
@@ -47,7 +47,7 @@ FAvaModifiersEditorStyle::FAvaModifiersEditorStyle()
 
 	FSlateStyleRegistry::RegisterSlateStyle(*this);
 
-	UActorModifierCoreSubsystem::OnModifierClassRegisteredDelegate.AddRaw(this, &FAvaModifiersEditorStyle::OnModifierClassRegistered);
+	UActorModifierCoreSubsystem::OnModifierClassRegistered().AddRaw(this, &FAvaModifiersEditorStyle::OnModifierClassRegistered);
 }
 
 const FSlateColor& FAvaModifiersEditorStyle::GetModifierCategoryColor(FName CategoryName)
@@ -80,7 +80,7 @@ void FAvaModifiersEditorStyle::OnModifierClassRegistered(const FActorModifierCor
 
 FAvaModifiersEditorStyle::~FAvaModifiersEditorStyle()
 {
-	UActorModifierCoreSubsystem::OnModifierClassRegisteredDelegate.RemoveAll(this);
+	UActorModifierCoreSubsystem::OnModifierClassRegistered().RemoveAll(this);
 
 	FSlateStyleRegistry::UnRegisterSlateStyle(*this);
 }
diff --git a/Engine/Plugins/Experimental/Avalanche/Source/AvalancheModifiersEditor/Private/Outliner/AvaOutlinerModifierProxy.cpp b/Engine/Plugins/Experimental/Avalanche/Source/AvalancheModifiersEditor/Private/Outliner/AvaOutlinerModifierProxy.cpp
--- a/Engine/Plugins/Experimental/Avalanche/Source/AvalancheModifiersEditor/Private/Outliner/AvaOutlinerModifierProxy.cpp
+++ b/Engine/Plugins/Experimental/Avalanche/Source/AvalancheModifiersEditor/Private/Outliner/AvaOutlinerModifierProxy.cpp
@@ -134,16 +134,18 @@ void FAvaOutlinerModifierProxy::GetProxiedItems(const TSharedRef<IAvaOutlinerIte
 void FAvaOutlinerModifierProxy::BindDelegates()
 {
 	UnbindDelegates();
-	UActorModifierCoreStack::OnModifierAddedDelegate.AddSP(this, &FAvaOutlinerModifierProxy::OnModifierStackUpdated);
-	UActorModifierCoreStack::OnModifierRemovedDelegate.AddSP(this, &FAvaOutlinerModifierProxy::OnModifierStackUpdated);
-	UActorModifierCoreStack::OnModifierMovedDelegate.AddSP(this, &FAvaOutlinerModifierProxy::OnModifierStackUpdated);
+	UActorModifierCoreStack::OnModifierAdded().AddSP(this, &FAvaOutlinerModifierProxy::OnModifierStackUpdated);
+	UActorModifierCoreStack::OnModifierRemoved().AddSP(this, &FAvaOutlinerModifierProxy::OnModifierStackUpdated);
+	UActorModifierCoreStack::OnModifierMoved().AddSP(this, &FAvaOutlinerModifierProxy::OnModifierStackUpdated);
+	UActorModifierCoreStack::OnModifierReplaced().AddSP(this, &FAvaOutlinerModifierProxy::OnModifierStackUpdated);
 }
 
 void FAvaOutlinerModifierProxy::UnbindDelegates()
 {
-	UActorModifierCoreStack::OnModifierAddedDelegate.RemoveAll(this);
-	UActorModifierCoreStack::OnModifierRemovedDelegate.RemoveAll(this);
-	UActorModifierCoreStack::OnModifierMovedDelegate.RemoveAll(this);
+	UActorModifierCoreStack::OnModifierAdded().RemoveAll(this);
+	UActorModifierCoreStack::OnModifierRemoved().RemoveAll(this);
+	UActorModifierCoreStack::OnModifierMoved().RemoveAll(this);
+	UActorModifierCoreStack::OnModifierReplaced().RemoveAll(this);
 }
 
 void FAvaOutlinerModifierProxy::OnModifierStackUpdated(UActorModifierCoreBase* ItemChanged)
