MotionDesign : ActorModifierCore
- Fixed replacing blueprint modifier would still leave metadata leftovers

diff --git a/Engine/Plugins/Experimental/ActorModifierCore/Source/ActorModifierCore/Private/Subsystems/ActorModifierCoreSubsystem.cpp b/Engine/Plugins/Experimental/ActorModifierCore/Source/ActorModifierCore/Private/Subsystems/ActorModifierCoreSubsystem.cpp
--- a/Engine/Plugins/Experimental/ActorModifierCore/Source/ActorModifierCore/Private/Subsystems/ActorModifierCoreSubsystem.cpp
+++ b/Engine/Plugins/Experimental/ActorModifierCore/Source/ActorModifierCore/Private/Subsystems/ActorModifierCoreSubsystem.cpp
@@ -56,7 +56,7 @@ void UActorModifierCoreSubsystem::Initialize(FSubsystemCollectionBase& Collectio
 	AssetRegistry.OnFilesLoaded().AddUObject(this, &UActorModifierCoreSubsystem::OnAssetRegistryFilesLoaded);
 	AssetRegistry.OnAssetAdded().AddUObject(this, &UActorModifierCoreSubsystem::OnAssetRegistryAssetAdded);
 	AssetRegistry.OnAssetRemoved().AddUObject(this, &UActorModifierCoreSubsystem::OnAssetRegistryAssetRemoved);
-	AssetRegistry.OnAssetUpdatedOnDisk().AddUObject(this, &UActorModifierCoreSubsystem::OnAssetRegistryAssetUpdated);
+	AssetRegistry.OnAssetUpdated().AddUObject(this, &UActorModifierCoreSubsystem::OnAssetRegistryAssetUpdated);
 
 #if WITH_EDITOR
 	FCoreUObjectDelegates::OnObjectsReplaced.AddUObject(this, &UActorModifierCoreSubsystem::OnBlueprintObjectsReplaced);
@@ -79,7 +79,7 @@ void UActorModifierCoreSubsystem::Deinitialize()
 		AssetRegistry.OnFilesLoaded().RemoveAll(this);
 		AssetRegistry.OnAssetAdded().RemoveAll(this);
 		AssetRegistry.OnAssetRemoved().RemoveAll(this);
-		AssetRegistry.OnInMemoryAssetDeleted().RemoveAll(this);
+		AssetRegistry.OnAssetUpdated().RemoveAll(this);
 	}
 
 #if WITH_EDITOR
@@ -103,7 +103,7 @@ void UActorModifierCoreSubsystem::OnInsertModifier(const FActorModifierCoreStack
 
 bool UActorModifierCoreSubsystem::RegisterModifierClass(const UClass* InModifierClass, bool bInOverrideIfExists)
 {
-	if (!IsValid(InModifierClass))
+	if (!InModifierClass)
 	{
 		return false;
 	}
@@ -166,15 +166,17 @@ bool UActorModifierCoreSubsystem::UnregisterModifierClass(const FName& InName)
 
 bool UActorModifierCoreSubsystem::UnregisterModifierClass(const UClass* InModifierClass)
 {
-	for (const TPair<FName, TSharedRef<FActorModifierCoreMetadata>>& ModifierMetadataPair : ModifiersMetadata)
+	bool bUnregistered = false;
+
+	for (TMap<FName, TSharedRef<FActorModifierCoreMetadata>>::TIterator It(ModifiersMetadata); It; ++It)
 	{
-		if (ModifierMetadataPair.Value->GetClass() == InModifierClass)
+		if (It->Value->GetClass() == InModifierClass)
 		{
-			return UnregisterModifierClass(ModifierMetadataPair.Key);
+			bUnregistered |= UnregisterModifierClass(It->Key);
 		}
 	}
 
-	return false;
+	return bUnregistered;
 }
 
 bool UActorModifierCoreSubsystem::IsRegisteredModifierClass(const FName& InName) const
@@ -1502,7 +1504,8 @@ void UActorModifierCoreSubsystem::OnBlueprintObjectsReplaced(const TMap<UObject*
 {
 	for (const TPair<UObject*, UObject*>& Replacement : InReplacements)
 	{
-		if (!Replacement.Key || !Replacement.Value || Replacement.Value->IsTemplate())
+		// Replace modifiers that are not the template
+		if (!Replacement.Key || !Replacement.Value || Replacement.Key->IsTemplate() || Replacement.Value->IsTemplate())
 		{
 			continue;
 		}
@@ -1513,7 +1516,9 @@ void UActorModifierCoreSubsystem::OnBlueprintObjectsReplaced(const TMap<UObject*
 		if (ReplacedModifier && ReplacementModifier)
 		{
 			UnregisterModifierClass(ReplacedModifier->GetClass());
-			RegisterModifierClass(ReplacementModifier->GetClass(), true);
+			UnregisterModifierClass(ReplacementModifier->GetClass());
+
+			RegisterModifierClass(ReplacementModifier->GetClass());
 
 			OnModifierReplacedDelegate.Broadcast(ReplacedModifier, ReplacementModifier);
 		}
@@ -1526,6 +1531,7 @@ void UActorModifierCoreSubsystem::RegisterModifierAsset(const FAssetData& InAsse
 	if (InAssetData.GetTagValue(TEXT("GeneratedClass"), GeneratedClassPath))
 	{
 		const UClass* GeneratedClass = LoadObject<UClass>(nullptr, *GeneratedClassPath);
+		UnregisterModifierClass(GeneratedClass);
 		RegisterModifierClass(GeneratedClass);
 	}
 }
