MotionDesign : ActorModifierCore
- Fixed registering blueprint modifier assets would cause long compiling in debug

diff --git a/Engine/Plugins/Experimental/ActorModifierCore/Source/ActorModifierCore/Private/Subsystems/ActorModifierCoreSubsystem.cpp b/Engine/Plugins/Experimental/ActorModifierCore/Source/ActorModifierCore/Private/Subsystems/ActorModifierCoreSubsystem.cpp
--- a/Engine/Plugins/Experimental/ActorModifierCore/Source/ActorModifierCore/Private/Subsystems/ActorModifierCoreSubsystem.cpp
+++ b/Engine/Plugins/Experimental/ActorModifierCore/Source/ActorModifierCore/Private/Subsystems/ActorModifierCoreSubsystem.cpp
@@ -1476,7 +1476,7 @@ void UActorModifierCoreSubsystem::OnAssetRegistryFilesLoaded()
 	const FAssetRegistryModule& AssetRegistryModule = FModuleManager::LoadModuleChecked<FAssetRegistryModule>("AssetRegistry");
 
 	TArray<FAssetData> Assets;
-	const FTopLevelAssetPath BlueprintPath(TEXT("/Script/ActorModifierCoreBlueprint.ActorModifierCoreBlueprint"));
+	const FTopLevelAssetPath BlueprintPath(BlueprintClass);
 	AssetRegistryModule.Get().GetAssetsByClass(BlueprintPath, Assets, /** Subclass */true);
 
 	for (const FAssetData& Asset : Assets)
@@ -1527,22 +1527,32 @@ void UActorModifierCoreSubsystem::OnBlueprintObjectsReplaced(const TMap<UObject*
 
 void UActorModifierCoreSubsystem::RegisterModifierAsset(const FAssetData& InAssetData)
 {
-	FString GeneratedClassPath;
-	if (InAssetData.GetTagValue(TEXT("GeneratedClass"), GeneratedClassPath))
+	if (InAssetData.AssetClassPath.ToString().Equals(BlueprintClass))
 	{
-		const UClass* GeneratedClass = LoadObject<UClass>(nullptr, *GeneratedClassPath);
-		UnregisterModifierClass(GeneratedClass);
-		RegisterModifierClass(GeneratedClass);
+		FString GeneratedClassPath;
+		if (InAssetData.GetTagValue(TEXT("GeneratedClass"), GeneratedClassPath))
+		{
+			if (const UClass* GeneratedClass = LoadObject<UClass>(nullptr, *GeneratedClassPath))
+			{
+				UnregisterModifierClass(GeneratedClass);
+				RegisterModifierClass(GeneratedClass);
+			}
+		}
 	}
 }
 
 void UActorModifierCoreSubsystem::UnregisterModifierAsset(const FAssetData& InAssetData)
 {
-	FString GeneratedClassPath;
-	if (InAssetData.GetTagValue(TEXT("GeneratedClass"), GeneratedClassPath))
+	if (InAssetData.AssetClassPath.ToString().Equals(BlueprintClass))
 	{
-		const UClass* GeneratedClass = LoadObject<UClass>(nullptr, *GeneratedClassPath);
-		UnregisterModifierClass(GeneratedClass);
+		FString GeneratedClassPath;
+		if (InAssetData.GetTagValue(TEXT("GeneratedClass"), GeneratedClassPath))
+		{
+			if (const UClass* GeneratedClass = LoadObject<UClass>(nullptr, *GeneratedClassPath))
+			{
+				UnregisterModifierClass(GeneratedClass);
+			}
+		}
 	}
 }
 
diff --git a/Engine/Plugins/Experimental/ActorModifierCore/Source/ActorModifierCore/Public/Subsystems/ActorModifierCoreSubsystem.h b/Engine/Plugins/Experimental/ActorModifierCore/Source/ActorModifierCore/Public/Subsystems/ActorModifierCoreSubsystem.h
--- a/Engine/Plugins/Experimental/ActorModifierCore/Source/ActorModifierCore/Public/Subsystems/ActorModifierCoreSubsystem.h
+++ b/Engine/Plugins/Experimental/ActorModifierCore/Source/ActorModifierCore/Public/Subsystems/ActorModifierCoreSubsystem.h
@@ -27,6 +27,8 @@ class UActorModifierCoreSubsystem : public UEngineSubsystem
 	friend struct FActorModifierCoreMetadata;
 
 public:
+	static constexpr const TCHAR* BlueprintClass = TEXT("/Script/ActorModifierCoreBlueprint.ActorModifierCoreBlueprint");
+
 	DECLARE_MULTICAST_DELEGATE_OneParam(FOnModifierClassRegistered, const FActorModifierCoreMetadata& /** ModifierMetadata */)
 	DECLARE_MULTICAST_DELEGATE_OneParam(FOnModifierStackRegistered, const UActorModifierCoreStack* /** ActorRootStack */)
 	DECLARE_MULTICAST_DELEGATE_TwoParams(FOnModifierReplaced, UActorModifierCoreBlueprintBase* /** PreviousModifier */, UActorModifierCoreBlueprintBase* /** ReplacementModifier */)
