MotionDesign : ActorModifierCore
- Add register blueprint modifier engine analytics

diff --git a/Engine/Plugins/Experimental/ActorModifierCore/Source/ActorModifierCore/Private/Subsystems/ActorModifierCoreSubsystem.cpp b/Engine/Plugins/Experimental/ActorModifierCore/Source/ActorModifierCore/Private/Subsystems/ActorModifierCoreSubsystem.cpp
--- a/Engine/Plugins/Experimental/ActorModifierCore/Source/ActorModifierCore/Private/Subsystems/ActorModifierCoreSubsystem.cpp
+++ b/Engine/Plugins/Experimental/ActorModifierCore/Source/ActorModifierCore/Private/Subsystems/ActorModifierCoreSubsystem.cpp
@@ -1473,6 +1473,8 @@ void UActorModifierCoreSubsystem::ScanForModifiers()
 
 void UActorModifierCoreSubsystem::OnAssetRegistryFilesLoaded()
 {
+	bFilesLoaded = true;
+
 	const FAssetRegistryModule& AssetRegistryModule = FModuleManager::LoadModuleChecked<FAssetRegistryModule>("AssetRegistry");
 
 	TArray<FAssetData> Assets;
@@ -1487,7 +1489,10 @@ void UActorModifierCoreSubsystem::OnAssetRegistryFilesLoaded()
 
 void UActorModifierCoreSubsystem::OnAssetRegistryAssetAdded(const FAssetData& InAssetData)
 {
-	RegisterModifierAsset(InAssetData);
+	if (bFilesLoaded)
+	{
+		RegisterModifierAsset(InAssetData);
+	}
 }
 
 void UActorModifierCoreSubsystem::OnAssetRegistryAssetUpdated(const FAssetData& InAssetData)
@@ -1535,7 +1540,18 @@ void UActorModifierCoreSubsystem::RegisterModifierAsset(const FAssetData& InAsse
 			if (const UClass* GeneratedClass = LoadObject<UClass>(nullptr, *GeneratedClassPath))
 			{
 				UnregisterModifierClass(GeneratedClass);
-				RegisterModifierClass(GeneratedClass);
+
+				if (RegisterModifierClass(GeneratedClass))
+				{
+#if WITH_EDITOR
+					if (FEngineAnalytics::IsAvailable())
+					{
+						TArray<FAnalyticsEventAttribute> Attributes;
+						Attributes.Emplace(TEXT("Class"), GeneratedClassPath);
+						FEngineAnalytics::GetProvider().RecordEvent(TEXT("Editor.Usage.ActorModifiers.RegisterBlueprintModifier"), Attributes);
+					}
+#endif
+				}
 			}
 		}
 	}
diff --git a/Engine/Plugins/Experimental/ActorModifierCore/Source/ActorModifierCore/Public/Subsystems/ActorModifierCoreSubsystem.h b/Engine/Plugins/Experimental/ActorModifierCore/Source/ActorModifierCore/Public/Subsystems/ActorModifierCoreSubsystem.h
--- a/Engine/Plugins/Experimental/ActorModifierCore/Source/ActorModifierCore/Public/Subsystems/ActorModifierCoreSubsystem.h
+++ b/Engine/Plugins/Experimental/ActorModifierCore/Source/ActorModifierCore/Public/Subsystems/ActorModifierCoreSubsystem.h
@@ -213,4 +213,7 @@ protected:
 
 	/** Stores modifiers providers per world/level, there should be only one provider per level */
 	TMap<TWeakObjectPtr<const ULevel>, TWeakObjectPtr<AActorModifierCoreSharedActor>> ModifierSharedProviders;
+
+	/** Asset Registry files loaded */
+	bool bFilesLoaded = false;
 };
