Motion Design: Added Sub Layer Task + Modifier for a simplified worfklow

#virtualized

diff --git a/Engine/Plugins/Experimental/ActorModifierCore/Source/ActorModifierCore/Private/Modifiers/Utilities/ActorModifierCoreLibrary.cpp b/Engine/Plugins/Experimental/ActorModifierCore/Source/ActorModifierCore/Private/Modifiers/Utilities/ActorModifierCoreLibrary.cpp
--- a/Engine/Plugins/Experimental/ActorModifierCore/Source/ActorModifierCore/Private/Modifiers/Utilities/ActorModifierCoreLibrary.cpp
+++ b/Engine/Plugins/Experimental/ActorModifierCore/Source/ActorModifierCore/Private/Modifiers/Utilities/ActorModifierCoreLibrary.cpp
@@ -8,7 +8,7 @@
 
 DEFINE_LOG_CATEGORY_STATIC(LogActorModifierCoreLibrary, Log, All);
 
-bool UActorModifierCoreLibrary::FindOrAddModifierStack(AActor* InActor, UActorModifierCoreStack*& OutModifierStack, bool bInCreateIfNone)
+bool UActorModifierCoreLibrary::FindModifierStack(AActor* InActor, UActorModifierCoreStack*& OutModifierStack, bool bInCreateIfNone)
 {
 	OutModifierStack = nullptr;
 	const UActorModifierCoreSubsystem* ModifierSubsystem = UActorModifierCoreSubsystem::Get();
@@ -460,6 +460,15 @@ bool UActorModifierCoreLibrary::FindModifiers(UActorModifierCoreStack* InModifie
 	return !OutFoundModifiers.IsEmpty();
 }
 
+UActorModifierCoreBase* UActorModifierCoreLibrary::FindModifierOfClass(UActorModifierCoreStack* InModifierStack, TSubclassOf<UActorModifierCoreBase> InModifierClass)
+{
+	if (IsValid(InModifierStack))
+	{
+		return InModifierStack->FindModifier(InModifierClass);
+	}
+	return nullptr;
+}
+
 bool UActorModifierCoreLibrary::GetSupportedModifiers(AActor* InActor, TSet<TSubclassOf<UActorModifierCoreBase>>& OutSupportedModifierClasses, EActorModifierCoreStackPosition InContextPosition, UActorModifierCoreBase* InContextModifier)
 {
 	const UActorModifierCoreSubsystem* ModifierSubsystem = UActorModifierCoreSubsystem::Get();
diff --git a/Engine/Plugins/Experimental/ActorModifierCore/Source/ActorModifierCore/Public/Modifiers/Utilities/ActorModifierCoreLibrary.h b/Engine/Plugins/Experimental/ActorModifierCore/Source/ActorModifierCore/Public/Modifiers/Utilities/ActorModifierCoreLibrary.h
--- a/Engine/Plugins/Experimental/ActorModifierCore/Source/ActorModifierCore/Public/Modifiers/Utilities/ActorModifierCoreLibrary.h
+++ b/Engine/Plugins/Experimental/ActorModifierCore/Source/ActorModifierCore/Public/Modifiers/Utilities/ActorModifierCoreLibrary.h
@@ -24,7 +24,7 @@ public:
 	 * @return true when a valid modifier stack is returned
 	 */
 	UFUNCTION(BlueprintCallable, Category="Motion Design|Modifiers|Utility")
-	static ACTORMODIFIERCORE_API bool FindOrAddModifierStack(AActor* InActor, UActorModifierCoreStack*& OutModifierStack, bool bInCreateIfNone = false);
+	static ACTORMODIFIERCORE_API bool FindModifierStack(AActor* InActor, UActorModifierCoreStack*& OutModifierStack, bool bInCreateIfNone = false);
 
 	/**
 	 * Creates and insert a new modifier into a modifier stack
@@ -208,6 +208,15 @@ public:
 	UFUNCTION(BlueprintCallable, Category="Motion Design|Modifiers|Utility")
 	static ACTORMODIFIERCORE_API bool FindModifiers(UActorModifierCoreStack* InModifierStack, const FActorModifierCoreSearchOperation& InOperation, TSet<UActorModifierCoreBase*>& OutFoundModifiers);
 
+	/**
+	 * Returns the first modifier of a specified class in the stack
+	 * @param InModifierStack The modifier stack to search
+	 * @param InModifierClass The class of the modifier to look for
+	 * @return the Modifier of the specified class, if any.
+	 */
+	UFUNCTION(BlueprintCallable, Category="Motion Design|Modifiers|Utility", meta=(DeterminesOutputType="InModifierClass"))
+	static ACTORMODIFIERCORE_API UActorModifierCoreBase* FindModifierOfClass(UActorModifierCoreStack* InModifierStack, TSubclassOf<UActorModifierCoreBase> InModifierClass);
+
 	/**
 	 * Gets all modifier classes supported by this actor at a specific position
 	 * @param InActor The actor to check for compatibility
