MotionDesign : ActorModifierCore
- Replacing the advanced find node by simpler find nodes to search for modifiers

diff --git a/Engine/Plugins/Experimental/ActorModifierCore/Config/BaseActorModifierCore.ini b/Engine/Plugins/Experimental/ActorModifierCore/Config/BaseActorModifierCore.ini
--- a/Engine/Plugins/Experimental/ActorModifierCore/Config/BaseActorModifierCore.ini
+++ b/Engine/Plugins/Experimental/ActorModifierCore/Config/BaseActorModifierCore.ini
@@ -6,4 +6,5 @@
 +ClassRedirects=(OldName="/Script/TagCollection.TagCollectionModifierStack",NewName="/Script/ActorModifierCore.ActorModifierCoreStack")
 +ClassRedirects=(OldName="/Script/TagCollectionEditor.TagCollectionModifierStackCustomization",NewName="/Script/ActorModifierCoreEditor.ActorModifierCoreEditorStackCustomization")
 +ClassRedirects=(OldName="/Script/TagCollectionEditor.TagCollectionModifierExtensionSubsystem",NewName="/Script/ActorModifierCoreEditor.ActorModifierCoreEditorSubsystem")
-+StructRedirects=(OldName="/Script/TagCollectionEditor.TagCollectionModifierPropertiesWrapper",NewName="/Script/ActorModifierCoreEditor.ActorModifierCoreEditorPropertiesWrapper")
\ No newline at end of file
++StructRedirects=(OldName="/Script/TagCollectionEditor.TagCollectionModifierPropertiesWrapper",NewName="/Script/ActorModifierCoreEditor.ActorModifierCoreEditorPropertiesWrapper")
++FunctionRedirects=(OldName="/Script/ActorModifierCore.ActorModifierCoreLibrary.FindModifierOfClass",NewName="/Script/ActorModifierCore.ActorModifierCoreLibrary.FindModifierByClass")
\ No newline at end of file
diff --git a/Engine/Plugins/Experimental/ActorModifierCore/Source/ActorModifierCore/Private/Modifiers/Utilities/ActorModifierCoreLibrary.cpp b/Engine/Plugins/Experimental/ActorModifierCore/Source/ActorModifierCore/Private/Modifiers/Utilities/ActorModifierCoreLibrary.cpp
--- a/Engine/Plugins/Experimental/ActorModifierCore/Source/ActorModifierCore/Private/Modifiers/Utilities/ActorModifierCoreLibrary.cpp
+++ b/Engine/Plugins/Experimental/ActorModifierCore/Source/ActorModifierCore/Private/Modifiers/Utilities/ActorModifierCoreLibrary.cpp
@@ -370,103 +370,54 @@ bool UActorModifierCoreLibrary::GetRequiredModifiers(UActorModifierCoreBase* InM
 	return true;
 }
 
-bool UActorModifierCoreLibrary::ContainsModifiers(UActorModifierCoreStack* InModifierStack, const FActorModifierCoreSearchOperation& InOperation)
+UActorModifierCoreBase* UActorModifierCoreLibrary::FindModifierByClass(UActorModifierCoreStack* InModifierStack, TSubclassOf<UActorModifierCoreBase> InModifierClass)
 {
-	TSet<UActorModifierCoreBase*> FoundModifiers;
-	return FindModifiers(InModifierStack, InOperation, FoundModifiers);
+	if (IsValid(InModifierStack))
+	{
+		return InModifierStack->FindModifier(InModifierClass);
+	}
+
+	return nullptr;
 }
 
-bool UActorModifierCoreLibrary::FindModifiers(UActorModifierCoreStack* InModifierStack, const FActorModifierCoreSearchOperation& InOperation, TSet<UActorModifierCoreBase*>& OutFoundModifiers)
+UActorModifierCoreBase* UActorModifierCoreLibrary::FindModifierByName(const UActorModifierCoreStack* InModifierStack, FName InModifierName)
 {
-	OutFoundModifiers.Empty();
-
-	if (!IsValid(InModifierStack))
+	if (IsValid(InModifierStack))
 	{
-		return false;
+		return InModifierStack->FindModifier(InModifierName);
 	}
 
-	for (UActorModifierCoreBase* Modifier : InOperation.Modifiers)
-	{
-		if (!IsValid(Modifier))
-		{
-			continue;
-		}
-
-		const bool bFoundModifier = InModifierStack->ContainsModifier(Modifier);
-
-		if (bFoundModifier)
-		{
-			OutFoundModifiers.Add(Modifier);
-		}
-
-		if (!bFoundModifier && InOperation.SearchMode == EActorModifierCoreSearchMode::And)
-		{
-			return false;
-		}
-
-		if (bFoundModifier && InOperation.SearchMode == EActorModifierCoreSearchMode::Or)
-		{
-			return true;
-		}
-	}
+	return nullptr;
+}
 
-	for (const TSubclassOf<UActorModifierCoreBase>& ModifierClass : InOperation.ModifierClasses)
+TArray<UActorModifierCoreBase*> UActorModifierCoreLibrary::FindModifiersByClass(UActorModifierCoreStack* InModifierStack, TSubclassOf<UActorModifierCoreBase> InModifierClass)
+{
+	if (IsValid(InModifierStack))
 	{
-		const UClass* Class = ModifierClass.Get();
-
-		if (!IsValid(Class))
-		{
-			continue;
-		}
-
-		UActorModifierCoreBase* FoundModifier = InModifierStack->FindModifier(Class);
-
-		if (FoundModifier)
-		{
-			OutFoundModifiers.Add(FoundModifier);
-		}
-
-		if (!FoundModifier && InOperation.SearchMode == EActorModifierCoreSearchMode::And)
-		{
-			return false;
-		}
-
-		if (FoundModifier && InOperation.SearchMode == EActorModifierCoreSearchMode::Or)
-		{
-			return true;
-		}
+		return InModifierStack->FindModifiers(InModifierClass);
 	}
 
-	for (const FName& ModifierName : InOperation.ModifierNames)
-	{
-		UActorModifierCoreBase* FoundModifier = InModifierStack->FindModifier(ModifierName);
-
-		if (FoundModifier)
-		{
-			OutFoundModifiers.Add(FoundModifier);
-		}
-
-		if (!FoundModifier && InOperation.SearchMode == EActorModifierCoreSearchMode::And)
-		{
-			return false;
-		}
+	return {};
+}
 
-		if (FoundModifier && InOperation.SearchMode == EActorModifierCoreSearchMode::Or)
-		{
-			return true;
-		}
+TArray<UActorModifierCoreBase*> UActorModifierCoreLibrary::FindModifiersByName(UActorModifierCoreStack* InModifierStack, FName InModifierName)
+{
+	if (IsValid(InModifierStack))
+	{
+		return InModifierStack->FindModifiers(InModifierName);
 	}
 
-	return !OutFoundModifiers.IsEmpty();
+	return {};
 }
 
-UActorModifierCoreBase* UActorModifierCoreLibrary::FindModifierOfClass(UActorModifierCoreStack* InModifierStack, TSubclassOf<UActorModifierCoreBase> InModifierClass)
+bool UActorModifierCoreLibrary::ContainsModifier(UActorModifierCoreStack* InModifierStack, UActorModifierCoreBase* InModifier)
 {
 	if (IsValid(InModifierStack))
 	{
-		return InModifierStack->FindModifier(InModifierClass);
+		return InModifierStack->ContainsModifier(InModifier);
 	}
-	return nullptr;
+
+	return false;
 }
 
 bool UActorModifierCoreLibrary::GetSupportedModifiers(AActor* InActor, TSet<TSubclassOf<UActorModifierCoreBase>>& OutSupportedModifierClasses, EActorModifierCoreStackPosition InContextPosition, UActorModifierCoreBase* InContextModifier)
diff --git a/Engine/Plugins/Experimental/ActorModifierCore/Source/ActorModifierCore/Public/Modifiers/Utilities/ActorModifierCoreLibrary.h b/Engine/Plugins/Experimental/ActorModifierCore/Source/ActorModifierCore/Public/Modifiers/Utilities/ActorModifierCoreLibrary.h
--- a/Engine/Plugins/Experimental/ActorModifierCore/Source/ActorModifierCore/Public/Modifiers/Utilities/ActorModifierCoreLibrary.h
+++ b/Engine/Plugins/Experimental/ActorModifierCore/Source/ActorModifierCore/Public/Modifiers/Utilities/ActorModifierCoreLibrary.h
@@ -190,32 +190,49 @@ public:
 	static ACTORMODIFIERCORE_API bool GetRequiredModifiers(UActorModifierCoreBase* InModifier, TSet<UActorModifierCoreBase*>& OutModifiers);
 
 	/**
-	 * Searches modifiers in the stack based on the search options provided
-	 * @param InModifierStack The modifier stack used for the search
-	 * @param InOperation The data for this operation
-	 * @return true when the operation was successful 
+	 * Returns the first modifier of a specified class in the stack
+	 * @param InModifierStack The modifier stack to search
+	 * @param InModifierClass The class of the modifier to look for
+	 * @return the modifier of the specified class, if any.
 	 */
-	UFUNCTION(BlueprintCallable, Category="Motion Design|Modifiers|Utility")
-	static ACTORMODIFIERCORE_API bool ContainsModifiers(UActorModifierCoreStack* InModifierStack, const FActorModifierCoreSearchOperation& InOperation);
+	UFUNCTION(BlueprintCallable, Category="Motion Design|Modifiers|Utility", meta=(DeterminesOutputType="InModifierClass"))
+	static ACTORMODIFIERCORE_API UActorModifierCoreBase* FindModifierByClass(UActorModifierCoreStack* InModifierStack, TSubclassOf<UActorModifierCoreBase> InModifierClass);
 
 	/**
-	 * Searches modifiers in the stack based on the search options provided
-	 * @param InModifierStack The modifier stack used for the search
-	 * @param InOperation The data for this operation
-	 * @param OutFoundModifiers [Out] The modifiers that matches the search filters
-	 * @return true when the operation was successful
+	 * Returns the first modifier with specified name in the stack
+	 * @param InModifierStack The modifier stack to search
+	 * @param InModifierName The name of the modifier to look for
+	 * @return the modifier with specified name, if any.
 	 */
 	UFUNCTION(BlueprintCallable, Category="Motion Design|Modifiers|Utility")
-	static ACTORMODIFIERCORE_API bool FindModifiers(UActorModifierCoreStack* InModifierStack, const FActorModifierCoreSearchOperation& InOperation, TSet<UActorModifierCoreBase*>& OutFoundModifiers);
+	static ACTORMODIFIERCORE_API UActorModifierCoreBase* FindModifierByName(const UActorModifierCoreStack* InModifierStack, FName InModifierName);
 
 	/**
-	 * Returns the first modifier of a specified class in the stack
+	 * Finds all modifiers with specified class in the stack
 	 * @param InModifierStack The modifier stack to search
 	 * @param InModifierClass The class of the modifier to look for
-	 * @return the Modifier of the specified class, if any.
+	 * @return the modifiers with specified class, if any.
 	 */
 	UFUNCTION(BlueprintCallable, Category="Motion Design|Modifiers|Utility", meta=(DeterminesOutputType="InModifierClass"))
-	static ACTORMODIFIERCORE_API UActorModifierCoreBase* FindModifierOfClass(UActorModifierCoreStack* InModifierStack, TSubclassOf<UActorModifierCoreBase> InModifierClass);
+	static ACTORMODIFIERCORE_API TArray<UActorModifierCoreBase*> FindModifiersByClass(UActorModifierCoreStack* InModifierStack, TSubclassOf<UActorModifierCoreBase> InModifierClass);
+
+	/**
+	 * Finds all modifiers with specified class in the stack
+	 * @param InModifierStack The modifier stack to search
+	 * @param InModifierName The name of the modifier to look for
+	 * @return the modifiers with specified name, if any.
+	 */
+	UFUNCTION(BlueprintCallable, Category="Motion Design|Modifiers|Utility")
+	static ACTORMODIFIERCORE_API TArray<UActorModifierCoreBase*> FindModifiersByName(UActorModifierCoreStack* InModifierStack, FName InModifierName);
+
+	/**
+	 * Checks if a modifier is contained in the stack
+	 * @param InModifierStack The modifier stack to search
+	 * @param InModifier The modifier to look for
+	 * @return true if the modifier is contained within that stack
+	 */
+	UFUNCTION(BlueprintCallable, Category="Motion Design|Modifiers|Utility")
+	static ACTORMODIFIERCORE_API bool ContainsModifier(UActorModifierCoreStack* InModifierStack, UActorModifierCoreBase* InModifier);
 
 	/**
 	 * Gets all modifier classes supported by this actor at a specific position
diff --git a/Engine/Plugins/Experimental/ActorModifierCore/Source/ActorModifierCore/Public/Modifiers/Utilities/ActorModifierCoreLibraryDefs.h b/Engine/Plugins/Experimental/ActorModifierCore/Source/ActorModifierCore/Public/Modifiers/Utilities/ActorModifierCoreLibraryDefs.h
--- a/Engine/Plugins/Experimental/ActorModifierCore/Source/ActorModifierCore/Public/Modifiers/Utilities/ActorModifierCoreLibraryDefs.h
+++ b/Engine/Plugins/Experimental/ActorModifierCore/Source/ActorModifierCore/Public/Modifiers/Utilities/ActorModifierCoreLibraryDefs.h
@@ -7,15 +7,6 @@
 
 class UActorModifierCoreBase;
 
-UENUM(BlueprintType)
-enum class EActorModifierCoreSearchMode : uint8
-{
-	/** True when one condition is met */
-	Or,
-	/** True when all condition are met */
-	And
-};
-
 USTRUCT(BlueprintType)
 struct FActorModifierCoreInsertOperation
 {
@@ -71,22 +62,4 @@ struct FActorModifierCoreRemoveOperation
 
 	UPROPERTY(BlueprintReadWrite, Category="Motion Design|Modifiers|Utility")
 	bool bRemoveDependencies = false;
-};
-
-USTRUCT(BlueprintType)
-struct FActorModifierCoreSearchOperation
-{
-	GENERATED_BODY()
-
-	UPROPERTY(BlueprintReadWrite, Category="Motion Design|Modifiers|Utility")
-	EActorModifierCoreSearchMode SearchMode = EActorModifierCoreSearchMode::Or;
-
-	UPROPERTY(BlueprintReadWrite, Category="Motion Design|Modifiers|Utility")
-	TSet<TSubclassOf<UActorModifierCoreBase>> ModifierClasses;
-
-	UPROPERTY(BlueprintReadWrite, Category="Motion Design|Modifiers|Utility")
-	TSet<FName> ModifierNames;
-
-	UPROPERTY(BlueprintReadWrite, Category="Motion Design|Modifiers|Utility")
-	TSet<TObjectPtr<UActorModifierCoreBase>> Modifiers;
 };
\ No newline at end of file
