MotionDesign : Modifiers
- Optimizations to reduce actor modifiers updates within a stack

diff --git a/Engine/Plugins/Experimental/ActorModifierCore/Source/ActorModifierCore/Private/Modifiers/ActorModifierCoreBase.cpp b/Engine/Plugins/Experimental/ActorModifierCore/Source/ActorModifierCore/Private/Modifiers/ActorModifierCoreBase.cpp
--- a/Engine/Plugins/Experimental/ActorModifierCore/Source/ActorModifierCore/Private/Modifiers/ActorModifierCoreBase.cpp
+++ b/Engine/Plugins/Experimental/ActorModifierCore/Source/ActorModifierCore/Private/Modifiers/ActorModifierCoreBase.cpp
@@ -647,19 +647,6 @@ void UActorModifierCoreBase::LogModifier(const FString& InLog, bool bInForce) co
 	}
 }
 
-void UActorModifierCoreBase::TickModifier(float InDeltaTime)
-{
-	if (!Metadata->IsTickAllowed())
-	{
-		return;
-	}
-
-	if (IsModifierDirtyable())
-	{
-		MarkModifierDirty();
-	}
-}
-
 void UActorModifierCoreBase::PostModifierCreation(UActorModifierCoreStack* InStack)
 {
 	// initialize once, called by the subsystem itself
diff --git a/Engine/Plugins/Experimental/ActorModifierCore/Source/ActorModifierCore/Private/Modifiers/ActorModifierCoreStack.cpp b/Engine/Plugins/Experimental/ActorModifierCore/Source/ActorModifierCore/Private/Modifiers/ActorModifierCoreStack.cpp
--- a/Engine/Plugins/Experimental/ActorModifierCore/Source/ActorModifierCore/Private/Modifiers/ActorModifierCoreStack.cpp
+++ b/Engine/Plugins/Experimental/ActorModifierCore/Source/ActorModifierCore/Private/Modifiers/ActorModifierCoreStack.cpp
@@ -1832,32 +1832,40 @@ void UActorModifierCoreStack::OnModifierCDOSetup(FActorModifierCoreMetadata& InM
 	InMetadata.SetName(TEXT("Stack"));
 }
 
-bool UActorModifierCoreStack::IsModifierDirtyable() const
+void UActorModifierCoreStack::TickModifier(float InDelta) const
 {
-	if (!IsModifierInitialized()
+	if (!IsRootStack()
 		|| !IsModifierEnabled()
-		|| IsModifierDirty()
-		|| !IsModifierIdle())
+		|| !IsModifierInitialized()
+		|| !IsModifierIdle()
+		|| !Metadata->IsTickAllowed()
+		|| IsModifierDirty())
 	{
-		return false;
+		return;
 	}
 
-	// Tickable modifiers can mark stack dirty by returning true
-	for (const TObjectPtr<UActorModifierCoreBase>& Modifier : Modifiers)
+	// Skips nested stack when processing
+	UActorModifierCoreBase* FirstDirtyModifier = nullptr;
+	ProcessFunction([&FirstDirtyModifier](const UActorModifierCoreBase* InModifier)->bool
 	{
-		if (IsValid(Modifier)
-			&& Modifier->IsModifierEnabled()
-			&& !Modifier->IsModifierDirty())
+		if (IsValid(InModifier)
+			&& InModifier->IsModifierEnabled()
+			&& !InModifier->IsModifierDirty()
+			&& InModifier->Metadata->IsTickAllowed()
+			&& InModifier->IsModifierDirtyable())
 		{
-			if (Modifier->IsModifierDirtyable())
-			{
-				// Stop here all modifiers below will re-execute
-				return true;
-			}
+			FirstDirtyModifier = const_cast<UActorModifierCoreBase*>(InModifier);
+			// break loop
+			return false;
 		}
-	}
 
-	return false;
+		return true;
+	});
+
+	if (FirstDirtyModifier)
+	{
+		FirstDirtyModifier->MarkModifierDirty();
+	}
 }
 
 #undef LOCTEXT_NAMESPACE
diff --git a/Engine/Plugins/Experimental/ActorModifierCore/Source/ActorModifierCore/Public/Modifiers/ActorModifierCoreBase.h b/Engine/Plugins/Experimental/ActorModifierCore/Source/ActorModifierCore/Public/Modifiers/ActorModifierCoreBase.h
--- a/Engine/Plugins/Experimental/ActorModifierCore/Source/ActorModifierCore/Public/Modifiers/ActorModifierCoreBase.h
+++ b/Engine/Plugins/Experimental/ActorModifierCore/Source/ActorModifierCore/Public/Modifiers/ActorModifierCoreBase.h
@@ -237,7 +237,7 @@ protected:
 	/** Checks whether this modifier is ready to run, called before this modifier is executed */
 	virtual bool IsModifierReady() const { return true; }
 
-	/** You can do some additional lightweights checks here in case you want to dirty the modifier before the stack runs an update */
+	/** You can do additional lightweights checks here in case you want to dirty the modifier instead of waiting for an event, tick must be enabled in the metadata */
 	virtual bool IsModifierDirtyable() const { return false; }
 
 	/** Override in child classes, called before applying this modifier */
@@ -358,9 +358,6 @@ private:
 	/** Execute a const function on this modifier, only to read data */
 	ACTORMODIFIERCORE_API virtual bool ProcessFunction(TFunctionRef<bool(const UActorModifierCoreBase*)> InFunction, const FActorModifierCoreStackSearchOp& InSearchOptions) const;
 
-	/** INTERNAL USE ONLY, allows tickable modifier to mark themselves dirty */
-	void TickModifier(float InDeltaTime);
-
 	/** INTERNAL USE ONLY, called by the stack only to unapply this modifier if it was applied */
 	void Unapply();
 
diff --git a/Engine/Plugins/Experimental/ActorModifierCore/Source/ActorModifierCore/Public/Modifiers/ActorModifierCoreStack.h b/Engine/Plugins/Experimental/ActorModifierCore/Source/ActorModifierCore/Public/Modifiers/ActorModifierCoreStack.h
--- a/Engine/Plugins/Experimental/ActorModifierCore/Source/ActorModifierCore/Public/Modifiers/ActorModifierCoreStack.h
+++ b/Engine/Plugins/Experimental/ActorModifierCore/Source/ActorModifierCore/Public/Modifiers/ActorModifierCoreStack.h
@@ -18,6 +18,7 @@ class UActorModifierCoreStack : public UActorModifierCoreBase
 	GENERATED_BODY()
 
 	friend class UActorModifierCoreBase;
+	friend class UActorModifierCoreComponent;
 	friend class UActorModifierCoreSubsystem;
 
 	friend class FActorModifierCoreEditorDetailCustomization;
@@ -207,8 +208,8 @@ private:
 	/** Sets the stack to receive tick events */
 	virtual void OnModifierCDOSetup(FActorModifierCoreMetadata& InMetadata) override;
 
-	/** Checks if any modifier within the stack can be set dirty */
-	virtual bool IsModifierDirtyable() const override;
+	/** INTERNAL USE ONLY, allows tickable modifier to mark themselves dirty when checking IsModifierDirtyable */
+	void TickModifier(float InDelta) const;
 
 	/** Unregister this stack to the subsystem when this actor is destroyed */
 	UFUNCTION()
diff --git a/Engine/Plugins/Experimental/Avalanche/Source/AvalancheMask/Private/Handling/AvaMaskText3DActorMaterialCollectionHandle.cpp b/Engine/Plugins/Experimental/Avalanche/Source/AvalancheMask/Private/Handling/AvaMaskText3DActorMaterialCollectionHandle.cpp
--- a/Engine/Plugins/Experimental/Avalanche/Source/AvalancheMask/Private/Handling/AvaMaskText3DActorMaterialCollectionHandle.cpp
+++ b/Engine/Plugins/Experimental/Avalanche/Source/AvalancheMask/Private/Handling/AvaMaskText3DActorMaterialCollectionHandle.cpp
@@ -708,7 +708,6 @@ bool FAvaMaskAvaTextActorMaterialCollectionHandle::ApplyOriginalState(const FStr
 			{
 				TextComponent->SetTranslucencyStyle(HandleData->OriginalTranslucencyStyle);	
 			}
-			TextComponent->RefreshMaterialInstances();
 			
 			Super::ApplyOriginalState(InHandleData);
 			
diff --git a/Engine/Plugins/Experimental/Avalanche/Source/AvalancheModifiers/Private/Extensions/AvaRenderStateUpdateModifierExtension.cpp b/Engine/Plugins/Experimental/Avalanche/Source/AvalancheModifiers/Private/Extensions/AvaRenderStateUpdateModifierExtension.cpp
--- a/Engine/Plugins/Experimental/Avalanche/Source/AvalancheModifiers/Private/Extensions/AvaRenderStateUpdateModifierExtension.cpp
+++ b/Engine/Plugins/Experimental/Avalanche/Source/AvalancheModifiers/Private/Extensions/AvaRenderStateUpdateModifierExtension.cpp
@@ -3,6 +3,7 @@
 #include "Extensions/AvaRenderStateUpdateModifierExtension.h"
 
 #include "AvaModifiersActorUtils.h"
+#include "Components/PrimitiveComponent.h"
 #include "Containers/Ticker.h"
 #include "Modifiers/ActorModifierCoreBase.h"
 
@@ -94,6 +95,11 @@ void FAvaRenderStateUpdateModifierExtension::OnExtensionDisabled(EActorModifierC
 
 void FAvaRenderStateUpdateModifierExtension::OnRenderStateDirty(UActorComponent& InComponent)
 {
+	if (!InComponent.IsA<UPrimitiveComponent>())
+	{
+		return;
+	}
+
 	const AActor* ModifierActor = GetModifierActor();
 	AActor* ActorDirty = InComponent.GetOwner();
 
diff --git a/Engine/Plugins/Experimental/Avalanche/Source/AvalancheModifiers/Private/Modifiers/AvaAutoFollowModifier.cpp b/Engine/Plugins/Experimental/Avalanche/Source/AvalancheModifiers/Private/Modifiers/AvaAutoFollowModifier.cpp
--- a/Engine/Plugins/Experimental/Avalanche/Source/AvalancheModifiers/Private/Modifiers/AvaAutoFollowModifier.cpp
+++ b/Engine/Plugins/Experimental/Avalanche/Source/AvalancheModifiers/Private/Modifiers/AvaAutoFollowModifier.cpp
@@ -113,6 +113,12 @@ void UAvaAutoFollowModifier::Apply()
 		return;
 	}
 
+	if (ModifyActor->IsAttachedTo(FollowedActor))
+	{
+		Fail(LOCTEXT("InvalidReferenceActor", "Followed actor cannot be a parent of modified actor"));
+		return;
+	}
+
 	const FVector FollowAxisVector = FAvaModifiersActorUtils::GetVectorAxis(FollowedAxis);
 	if (FollowAxisVector.IsNearlyZero())
 	{
@@ -319,7 +325,7 @@ void UAvaAutoFollowModifier::OnTransformUpdated(AActor* InActor, bool bInParentM
 	const bool bIsAttachedToReferenceActor = InActor->IsAttachedTo(FollowedActor);
 	const bool bIsReferenceActor = InActor == FollowedActor;
 	
-	if (IsValid(FollowedActor) && (bIsAttachedToReferenceActor || bIsReferenceActor))
+	if (!bInParentMoved && IsValid(FollowedActor) && (bIsAttachedToReferenceActor || bIsReferenceActor))
 	{
 		MarkModifierDirty();
 	}
@@ -397,7 +403,11 @@ void UAvaAutoFollowModifier::OnSceneTreeTrackedActorChildrenChanged(int32 InIdx,
 
 void UAvaAutoFollowModifier::OnReferenceActorChanged()
 {
-	if (ReferenceActor.ReferenceActorWeak.Get() == GetModifiedActor())
+	const AActor* FollowerActor = GetModifiedActor();
+	const AActor* TrackedActor = ReferenceActor.ReferenceActorWeak.Get();
+
+	if (TrackedActor == FollowerActor
+		|| FollowerActor->IsAttachedTo(TrackedActor))
 	{
 		ReferenceActor.ReferenceActorWeak = nullptr;
 	}
diff --git a/Engine/Plugins/Experimental/Avalanche/Source/AvalancheModifiers/Private/Modifiers/AvaAutoSizeModifier.cpp b/Engine/Plugins/Experimental/Avalanche/Source/AvalancheModifiers/Private/Modifiers/AvaAutoSizeModifier.cpp
--- a/Engine/Plugins/Experimental/Avalanche/Source/AvalancheModifiers/Private/Modifiers/AvaAutoSizeModifier.cpp
+++ b/Engine/Plugins/Experimental/Avalanche/Source/AvalancheModifiers/Private/Modifiers/AvaAutoSizeModifier.cpp
@@ -25,7 +25,18 @@ bool UAvaAutoSizeModifier::IsModifierDirtyable() const
 
 	const FBox ReferenceActorLocalBounds = FAvaModifiersActorUtils::GetActorsBounds(TrackedActor, bIncludeChildren, true);
 
-	if (ReferenceActorLocalBounds.Equals(CachedReferenceBounds, 0.01))
+	const FVector FollowAxisVector = FAvaModifiersActorUtils::GetVectorAxis(FollowedAxis);
+
+	// If we follow reference actor check bounds
+	if (!FollowAxisVector.IsNearlyZero())
+	{
+		if (ReferenceActorLocalBounds.Equals(CachedReferenceBounds, 0.01))
+		{
+			return Super::IsModifierDirtyable();
+		}
+	}
+	// If we don't follow reference actor check size
+	else if (ReferenceActorLocalBounds.GetSize().Equals(CachedReferenceBounds.GetSize(), 0.01))
 	{
 		return Super::IsModifierDirtyable();
 	}
diff --git a/Engine/Plugins/Experimental/Avalanche/Source/AvalancheModifiers/Private/Modifiers/AvaVisibilityModifier.cpp b/Engine/Plugins/Experimental/Avalanche/Source/AvalancheModifiers/Private/Modifiers/AvaVisibilityModifier.cpp
--- a/Engine/Plugins/Experimental/Avalanche/Source/AvalancheModifiers/Private/Modifiers/AvaVisibilityModifier.cpp
+++ b/Engine/Plugins/Experimental/Avalanche/Source/AvalancheModifiers/Private/Modifiers/AvaVisibilityModifier.cpp
@@ -169,6 +169,11 @@ void UAvaVisibilityModifier::OnModifierCDOSetup(FActorModifierCoreMetadata& InMe
 #endif
 }
 
+void UAvaVisibilityModifier::OnModifiedActorTransformed()
+{
+	// Overwrite parent class behaviour don't do anything when moved
+}
+
 void UAvaVisibilityModifier::OnRenderStateUpdated(AActor* InActor, UActorComponent* InComponent)
 {
 	Super::OnRenderStateUpdated(InActor, InComponent);
diff --git a/Engine/Plugins/Experimental/Avalanche/Source/AvalancheModifiers/Public/Modifiers/AvaVisibilityModifier.h b/Engine/Plugins/Experimental/Avalanche/Source/AvalancheModifiers/Public/Modifiers/AvaVisibilityModifier.h
--- a/Engine/Plugins/Experimental/Avalanche/Source/AvalancheModifiers/Public/Modifiers/AvaVisibilityModifier.h
+++ b/Engine/Plugins/Experimental/Avalanche/Source/AvalancheModifiers/Public/Modifiers/AvaVisibilityModifier.h
@@ -63,6 +63,7 @@ protected:
 
 	//~ Begin UActorModifierCoreBase
 	virtual void OnModifierCDOSetup(FActorModifierCoreMetadata& InMetadata) override;
+	virtual void OnModifiedActorTransformed() override;
 	virtual void Apply() override;
 	//~ End UActorModifierCoreBase
 
diff --git a/Engine/Plugins/Experimental/Avalanche/Source/AvalancheShapes/Private/AvaShapeActor.cpp b/Engine/Plugins/Experimental/Avalanche/Source/AvalancheShapes/Private/AvaShapeActor.cpp
--- a/Engine/Plugins/Experimental/Avalanche/Source/AvalancheShapes/Private/AvaShapeActor.cpp
+++ b/Engine/Plugins/Experimental/Avalanche/Source/AvalancheShapes/Private/AvaShapeActor.cpp
@@ -11,15 +11,16 @@ const FName AAvaShapeActor::ShapeComponentName("Shape Component");
 // Sets default values
 AAvaShapeActor::AAvaShapeActor()
 {
-	// Set this actor to call Tick() every frame.  You can turn this off to improve performance if you don't need it.
 	SetCanBeDamaged(false);
-	PrimaryActorTick.bCanEverTick          = true;
-	PrimaryActorTick.bStartWithTickEnabled = true;
-	PrimaryActorTick.bTickEvenWhenPaused   = true;
+	PrimaryActorTick.bCanEverTick          = false;
+	PrimaryActorTick.bStartWithTickEnabled = false;
+	PrimaryActorTick.bTickEvenWhenPaused   = false;
 
 	ShapeMeshComponent = CreateDefaultSubobject<UDynamicMeshComponent>(ShapeComponentName);
 	ShapeMeshComponent->SetCollisionProfileName(UCollisionProfile::BlockAll_ProfileName);
 	ShapeMeshComponent->SetGenerateOverlapEvents(true);
+	ShapeMeshComponent->SetCanEverAffectNavigation(false);
+
 	SetRootComponent(ShapeMeshComponent);
 
 	bFinishedCreation = false;
diff --git a/Engine/Plugins/Experimental/Avalanche/Source/AvalancheShapes/Private/DynamicMeshes/AvaShapeDynMeshBase.cpp b/Engine/Plugins/Experimental/Avalanche/Source/AvalancheShapes/Private/DynamicMeshes/AvaShapeDynMeshBase.cpp
--- a/Engine/Plugins/Experimental/Avalanche/Source/AvalancheShapes/Private/DynamicMeshes/AvaShapeDynMeshBase.cpp
+++ b/Engine/Plugins/Experimental/Avalanche/Source/AvalancheShapes/Private/DynamicMeshes/AvaShapeDynMeshBase.cpp
@@ -51,11 +51,14 @@ UAvaShapeDynamicMeshBase::UAvaShapeDynamicMeshBase(const FLinearColor& InVertexC
 	, bVerticesDirty(false)
 	, bColorsDirty(false)
 {
+	if (!IsTemplate())
+	{
 #if WITH_EDITOR
-	FEditorDelegates::OnApplyObjectToActor.AddUObject(this, &UAvaShapeDynamicMeshBase::OnAssetDropped);
+		FEditorDelegates::OnApplyObjectToActor.AddUObject(this, &UAvaShapeDynamicMeshBase::OnAssetDropped);
 #endif
 
-	FAvaShapeParametricMaterial::OnMaterialChanged().AddUObject(this, &UAvaShapeDynamicMeshBase::OnParametricMaterialChanged);
+		FAvaShapeParametricMaterial::OnMaterialChanged().AddUObject(this, &UAvaShapeDynamicMeshBase::OnParametricMaterialChanged);
+	}
 }
 
 EMaterialType& UAvaShapeDynamicMeshBase::GetMaterialType(int32 MeshIndex)
