MotionDesign : OperatorStack
- Adding header tooltip text to display for each item when hovering on the icon and label

diff --git a/Engine/Plugins/Experimental/ActorModifierCore/Source/ActorModifierCoreEditor/Private/Modifiers/Customizations/ActorModifierCoreEditorStackCustomization.cpp b/Engine/Plugins/Experimental/ActorModifierCore/Source/ActorModifierCoreEditor/Private/Modifiers/Customizations/ActorModifierCoreEditorStackCustomization.cpp
--- a/Engine/Plugins/Experimental/ActorModifierCore/Source/ActorModifierCoreEditor/Private/Modifiers/Customizations/ActorModifierCoreEditorStackCustomization.cpp
+++ b/Engine/Plugins/Experimental/ActorModifierCore/Source/ActorModifierCoreEditor/Private/Modifiers/Customizations/ActorModifierCoreEditorStackCustomization.cpp
@@ -2,11 +2,15 @@
 
 #include "Modifiers/Customizations/ActorModifierCoreEditorStackCustomization.h"
 
+#include "ActorModifierCoreEditorStyle.h"
+#include "Contexts/OperatorStackEditorMenuContext.h"
 #include "Framework/Commands/GenericCommands.h"
 #include "Framework/Notifications/NotificationManager.h"
 #include "HAL/PlatformApplicationMisc.h"
 #include "IDetailTreeNode.h"
 #include "IPropertyRowGenerator.h"
+#include "Items/OperatorStackEditorGroupItem.h"
+#include "Items/OperatorStackEditorObjectItem.h"
 #include "JsonObjectConverter.h"
 #include "Modifiers/ActorModifierCoreBase.h"
 #include "Modifiers/ActorModifierCoreComponent.h"
@@ -16,10 +20,6 @@
 #include "Styling/SlateIconFinder.h"
 #include "Subsystems/ActorModifierCoreEditorSubsystem.h"
 #include "Subsystems/ActorModifierCoreSubsystem.h"
-#include "ActorModifierCoreEditorStyle.h"
-#include "Contexts/OperatorStackEditorMenuContext.h"
-#include "Items/OperatorStackEditorGroupItem.h"
-#include "Items/OperatorStackEditorObjectItem.h"
 #include "Widgets/Notifications/SNotificationList.h"
 #include "Widgets/Views/STableRow.h"
 
@@ -221,11 +221,13 @@ void UActorModifierCoreEditorStackCustomization::CustomizeItemHeader(const FOper
 
 		FSlateIcon ModifierIcon = FSlateIconFinder::FindIconForClass(UActorModifierCoreBase::StaticClass());
 		FLinearColor ModifierColor = FLinearColor::Transparent;
+		FText ModifierTooltip = FText::GetEmpty();
 
-		ModifierSubsystem->ProcessModifierMetadata(Modifier->GetModifierName(), [&ModifierIcon, &ModifierColor](const FActorModifierCoreMetadata& InMetadata)
+		ModifierSubsystem->ProcessModifierMetadata(Modifier->GetModifierName(), [&ModifierIcon, &ModifierColor, &ModifierTooltip](const FActorModifierCoreMetadata& InMetadata)
 		{
 			ModifierIcon = InMetadata.GetIcon();
 			ModifierColor = InMetadata.GetColor();
+			ModifierTooltip = InMetadata.GetDescription();
 			return true;
 		});
 
@@ -295,6 +297,7 @@ void UActorModifierCoreEditorStackCustomization::CustomizeItemHeader(const FOper
 			.SetExpandable(!bIsStack)
 			.SetIcon(ModifierIcon.GetIcon())
 			.SetLabel(FText::FromString(HeaderLabel))
+			.SetTooltip(ModifierTooltip)
 			.SetBorderColor(ModifierColor)
 			.SetProperty(ModifierEnableProperty)
 			.SetCommandList(Commands)
diff --git a/Engine/Plugins/Experimental/OperatorStack/Source/OperatorStack/Private/Builders/OperatorStackEditorHeaderBuilder.cpp b/Engine/Plugins/Experimental/OperatorStack/Source/OperatorStack/Private/Builders/OperatorStackEditorHeaderBuilder.cpp
--- a/Engine/Plugins/Experimental/OperatorStack/Source/OperatorStack/Private/Builders/OperatorStackEditorHeaderBuilder.cpp
+++ b/Engine/Plugins/Experimental/OperatorStack/Source/OperatorStack/Private/Builders/OperatorStackEditorHeaderBuilder.cpp
@@ -16,6 +16,12 @@ FOperatorStackEditorHeaderBuilder& FOperatorStackEditorHeaderBuilder::SetLabel(c
 	return *this;
 }
 
+FOperatorStackEditorHeaderBuilder& FOperatorStackEditorHeaderBuilder::SetTooltip(const FText& InTooltip)
+{
+	Tooltip = InTooltip;
+	return *this;
+}
+
 FOperatorStackEditorHeaderBuilder& FOperatorStackEditorHeaderBuilder::SetBorderColor(const FLinearColor& InColor)
 {
 	BorderColor = InColor;
diff --git a/Engine/Plugins/Experimental/OperatorStack/Source/OperatorStack/Private/Widgets/SOperatorStackEditorStack.cpp b/Engine/Plugins/Experimental/OperatorStack/Source/OperatorStack/Private/Widgets/SOperatorStackEditorStack.cpp
--- a/Engine/Plugins/Experimental/OperatorStack/Source/OperatorStack/Private/Widgets/SOperatorStackEditorStack.cpp
+++ b/Engine/Plugins/Experimental/OperatorStack/Source/OperatorStack/Private/Widgets/SOperatorStackEditorStack.cpp
@@ -168,6 +168,7 @@ TSharedPtr<SWidget> SOperatorStackEditorStack::GenerateHeaderWidget()
 		[
 			SNew(SImage)
 			.Image(Icon)
+			.ToolTipText(HeaderBuilder.GetTooltip())
 		];
 	}
 
@@ -182,6 +183,7 @@ TSharedPtr<SWidget> SOperatorStackEditorStack::GenerateHeaderWidget()
 		[
 			SNew(STextBlock)
 			.Text(HeaderBuilder.GetLabel())
+			.ToolTipText(HeaderBuilder.GetTooltip())
 			.Justification(ETextJustify::Center)
 		];
 	}
diff --git a/Engine/Plugins/Experimental/OperatorStack/Source/OperatorStack/Public/Builders/OperatorStackEditorHeaderBuilder.h b/Engine/Plugins/Experimental/OperatorStack/Source/OperatorStack/Public/Builders/OperatorStackEditorHeaderBuilder.h
--- a/Engine/Plugins/Experimental/OperatorStack/Source/OperatorStack/Public/Builders/OperatorStackEditorHeaderBuilder.h
+++ b/Engine/Plugins/Experimental/OperatorStack/Source/OperatorStack/Public/Builders/OperatorStackEditorHeaderBuilder.h
@@ -34,6 +34,9 @@ struct OPERATORSTACKEDITOR_API FOperatorStackEditorHeaderBuilder
 	/** Set the label to display for this item header */
 	FOperatorStackEditorHeaderBuilder& SetLabel(const FText& InLabel);
 
+	/** Set the tooltip to display for this item header */
+	FOperatorStackEditorHeaderBuilder& SetTooltip(const FText& InTooltip);
+
 	/** Sets the border color surrounding the item */
 	FOperatorStackEditorHeaderBuilder& SetBorderColor(const FLinearColor& InColor);
 
@@ -86,6 +89,11 @@ struct OPERATORSTACKEDITOR_API FOperatorStackEditorHeaderBuilder
 		return Label;
 	}
 
+	const FText& GetTooltip() const
+	{
+		return Tooltip;
+	}
+
 	const FLinearColor& GetBorderColor() const
 	{
 		return BorderColor;
@@ -178,6 +186,9 @@ protected:
 	/** On the left, after to icon */
 	FText Label;
 
+	/** Shown on the header */
+	FText Tooltip;
+
 	/** The border color surrounding the item */
 	FLinearColor BorderColor = FLinearColor::Transparent;
 
