MotionDesign : Modifiers
- Fix modifier crash due to invalid metadata in rundown (game mode)

diff --git a/Engine/Plugins/Experimental/ActorModifierCore/Source/ActorModifierCore/Private/Modifiers/ActorModifierCoreBase.cpp b/Engine/Plugins/Experimental/ActorModifierCore/Source/ActorModifierCore/Private/Modifiers/ActorModifierCoreBase.cpp
--- a/Engine/Plugins/Experimental/ActorModifierCore/Source/ActorModifierCore/Private/Modifiers/ActorModifierCoreBase.cpp
+++ b/Engine/Plugins/Experimental/ActorModifierCore/Source/ActorModifierCore/Private/Modifiers/ActorModifierCoreBase.cpp
@@ -652,8 +652,16 @@ void UActorModifierCoreBase::PostModifierCreation(UActorModifierCoreStack* InSta
 	// initialize once, called by the subsystem itself
 	if (GetModifierStack() == InStack)
 	{
-		if (const UActorModifierCoreBase* CDO = GetClass()->GetDefaultObject<UActorModifierCoreBase>())
+		if (UActorModifierCoreBase* CDO = GetClass()->GetDefaultObject<UActorModifierCoreBase>())
 		{
+			if (!CDO->Metadata.IsValid())
+			{
+				if (UActorModifierCoreSubsystem* ModifierSubsystem = UActorModifierCoreSubsystem::Get())
+				{
+					ModifierSubsystem->RegisterModifierClass(CDO->GetClass(), /** Override*/true);
+				}
+			}
+			
 			Metadata = CDO->Metadata;
 		}
 
@@ -768,6 +776,12 @@ void UActorModifierCoreBase::InitializeModifier(EActorModifierCoreEnableReason I
 		if (!Metadata.IsValid())
 		{
 			PostModifierCreation(GetModifierStack());
+
+			// Cannot proceed with invalid metadata
+			if (!Metadata.IsValid())
+			{
+				UE_LOG(LogActorModifierCoreBase, Fatal, TEXT("Invalid modifier metadata for instance of class %s"), *GetClass()->GetName())	
+			}
 		}
 
 		// Initialize profiler
diff --git a/Engine/Plugins/Experimental/ActorModifierCore/Source/ActorModifierCore/Private/Subsystems/ActorModifierCoreSubsystem.cpp b/Engine/Plugins/Experimental/ActorModifierCore/Source/ActorModifierCore/Private/Subsystems/ActorModifierCoreSubsystem.cpp
--- a/Engine/Plugins/Experimental/ActorModifierCore/Source/ActorModifierCore/Private/Subsystems/ActorModifierCoreSubsystem.cpp
+++ b/Engine/Plugins/Experimental/ActorModifierCore/Source/ActorModifierCore/Private/Subsystems/ActorModifierCoreSubsystem.cpp
@@ -125,6 +125,12 @@ bool UActorModifierCoreSubsystem::RegisterModifierClass(const UClass* InModifier
 	{
 		CDO->PostModifierCDOCreation();
 
+		if (!CDO->Metadata.IsValid() || !CDO->Metadata->GetClass())
+		{
+			UE_LOG(LogActorModifierCoreSubsystem, Warning, TEXT("[%s] Could not register modifier class : Modifier metadata is invalid"), *CDO->GetClass()->GetName())
+			return false;	
+		}
+		
 		if (CDO->IsModifierStack())
 		{
 			return false;
@@ -143,8 +149,6 @@ bool UActorModifierCoreSubsystem::RegisterModifierClass(const UClass* InModifier
 			UnregisterModifierClass(ModifierName);
 		}
 
-		check(CDO->Metadata.IsValid())
-
 		ModifiersMetadata.Add(ModifierName, CDO->Metadata.ToSharedRef());
 
 		UActorModifierCoreSubsystem::OnModifierClassRegisteredDelegate.Broadcast(*CDO->Metadata);
diff --git a/Engine/Plugins/Experimental/Avalanche/Source/Avalanche/Private/Framework/AvaGizmoComponent.cpp b/Engine/Plugins/Experimental/Avalanche/Source/Avalanche/Private/Framework/AvaGizmoComponent.cpp
--- a/Engine/Plugins/Experimental/Avalanche/Source/Avalanche/Private/Framework/AvaGizmoComponent.cpp
+++ b/Engine/Plugins/Experimental/Avalanche/Source/Avalanche/Private/Framework/AvaGizmoComponent.cpp
@@ -180,10 +180,13 @@ void UAvaGizmoComponent::ApplyGizmoValues()
 		}
 
 #if WITH_EDITOR
-		// Mark camera cut to remove temporal effects
-		for (FLevelEditorViewportClient* LevelVC : GEditor->GetLevelViewportClients())
+		if (GEditor)
 		{
-			LevelVC->SetIsCameraCut();
+			// Mark camera cut to remove temporal effects
+			for (FLevelEditorViewportClient* LevelVC : GEditor->GetLevelViewportClients())
+			{
+				LevelVC->SetIsCameraCut();
+			}
 		}
 #endif
 
