MotionDesign : ActorModifier
- Fix error: undefined symbol clang with modifier extension

diff --git a/Engine/Plugins/Experimental/ActorModifier/Source/ActorModifier/Public/Extensions/ActorModifierRenderStateUpdateExtension.h b/Engine/Plugins/Experimental/ActorModifier/Source/ActorModifier/Public/Extensions/ActorModifierRenderStateUpdateExtension.h
--- a/Engine/Plugins/Experimental/ActorModifier/Source/ActorModifier/Public/Extensions/ActorModifierRenderStateUpdateExtension.h
+++ b/Engine/Plugins/Experimental/ActorModifier/Source/ActorModifier/Public/Extensions/ActorModifierRenderStateUpdateExtension.h
@@ -55,8 +55,8 @@ public:
 
 protected:
 	//~ Begin FActorModifierCoreExtension
-	virtual void OnExtensionEnabled(EActorModifierCoreEnableReason InReason) override;
-	virtual void OnExtensionDisabled(EActorModifierCoreDisableReason InReason) override;
+	ACTORMODIFIER_API virtual void OnExtensionEnabled(EActorModifierCoreEnableReason InReason) override;
+	ACTORMODIFIER_API virtual void OnExtensionDisabled(EActorModifierCoreDisableReason InReason) override;
 	//~ End FActorModifierCoreExtension
 
 private:
diff --git a/Engine/Plugins/Experimental/ActorModifier/Source/ActorModifier/Public/Extensions/ActorModifierSceneTreeUpdateExtension.h b/Engine/Plugins/Experimental/ActorModifier/Source/ActorModifier/Public/Extensions/ActorModifierSceneTreeUpdateExtension.h
--- a/Engine/Plugins/Experimental/ActorModifier/Source/ActorModifier/Public/Extensions/ActorModifierSceneTreeUpdateExtension.h
+++ b/Engine/Plugins/Experimental/ActorModifier/Source/ActorModifier/Public/Extensions/ActorModifierSceneTreeUpdateExtension.h
@@ -149,9 +149,9 @@ public:
 
 protected:
 	//~ Begin FActorModifierCoreExtension
-	virtual void OnExtensionInitialized() override;
-	virtual void OnExtensionEnabled(EActorModifierCoreEnableReason InReason) override;
-	virtual void OnExtensionDisabled(EActorModifierCoreDisableReason InReason) override;
+	ACTORMODIFIER_API virtual void OnExtensionInitialized() override;
+	ACTORMODIFIER_API virtual void OnExtensionEnabled(EActorModifierCoreEnableReason InReason) override;
+	ACTORMODIFIER_API virtual void OnExtensionDisabled(EActorModifierCoreDisableReason InReason) override;
 	//~ End FActorModifierCoreExtension
 
 	TSet<TWeakObjectPtr<AActor>> GetChildrenActorsRecursive(const AActor* InActor) const;
diff --git a/Engine/Plugins/Experimental/ActorModifier/Source/ActorModifier/Public/Extensions/ActorModifierTransformUpdateExtension.h b/Engine/Plugins/Experimental/ActorModifier/Source/ActorModifier/Public/Extensions/ActorModifierTransformUpdateExtension.h
--- a/Engine/Plugins/Experimental/ActorModifier/Source/ActorModifier/Public/Extensions/ActorModifierTransformUpdateExtension.h
+++ b/Engine/Plugins/Experimental/ActorModifier/Source/ActorModifier/Public/Extensions/ActorModifierTransformUpdateExtension.h
@@ -40,8 +40,8 @@ public:
 
 protected:
 	//~ Begin FActorModifierCoreExtension
-	virtual void OnExtensionEnabled(EActorModifierCoreEnableReason InReason) override;
-	virtual void OnExtensionDisabled(EActorModifierCoreDisableReason InReason) override;
+	ACTORMODIFIER_API virtual void OnExtensionEnabled(EActorModifierCoreEnableReason InReason) override;
+	ACTORMODIFIER_API virtual void OnExtensionDisabled(EActorModifierCoreDisableReason InReason) override;
 	//~ End FActorModifierCoreExtension
 
 private:
diff --git a/Engine/Plugins/Experimental/ActorModifierCore/Source/ActorModifierCore/Public/Modifiers/ActorModifierCoreExtension.h b/Engine/Plugins/Experimental/ActorModifierCore/Source/ActorModifierCore/Public/Modifiers/ActorModifierCoreExtension.h
--- a/Engine/Plugins/Experimental/ActorModifierCore/Source/ActorModifierCore/Public/Modifiers/ActorModifierCoreExtension.h
+++ b/Engine/Plugins/Experimental/ActorModifierCore/Source/ActorModifierCore/Public/Modifiers/ActorModifierCoreExtension.h
@@ -25,13 +25,13 @@ public:
 
 protected:
 	/** Overwrite this, called when the extension is created */
-	ACTORMODIFIERCORE_API virtual void OnExtensionInitialized() {}
+	virtual void OnExtensionInitialized() {}
 	
 	/** Overwrite this, called when the modifier is enabled */
-	ACTORMODIFIERCORE_API virtual void OnExtensionEnabled(EActorModifierCoreEnableReason InReason) {}
+	virtual void OnExtensionEnabled(EActorModifierCoreEnableReason InReason) {}
 
 	/** Overwrite this, called when the modifier is disabled */
-	ACTORMODIFIERCORE_API virtual void OnExtensionDisabled(EActorModifierCoreDisableReason InReason) {}
+	virtual void OnExtensionDisabled(EActorModifierCoreDisableReason InReason) {}
 
 	/** Get the modifier actor */
 	ACTORMODIFIERCORE_API AActor* GetModifierActor() const;
