MotionDesign : Modifiers
- Align between modifier only saves location instead of whole transform
- Fix modifier shared transform object save/restore capability
- Removing boolean modifier async execution to avoid having one frame delay for animation
- Removing usage of promises and simplifying stack execution flow to avoid crash caused by promises not fullfilled

diff --git a/Engine/Plugins/Experimental/ActorModifier/Source/ActorModifier/Public/Shared/ActorModifierTransformShared.h b/Engine/Plugins/Experimental/ActorModifier/Source/ActorModifier/Public/Shared/ActorModifierTransformShared.h
--- a/Engine/Plugins/Experimental/ActorModifier/Source/ActorModifier/Public/Shared/ActorModifierTransformShared.h
+++ b/Engine/Plugins/Experimental/ActorModifier/Source/ActorModifier/Public/Shared/ActorModifierTransformShared.h
@@ -59,7 +59,7 @@ struct FActorModifierTransformSharedModifierState
 
 	/** Used to restore only what has changed */
 	UPROPERTY()
-	EActorModifierTransformSharedState SaveState = EActorModifierTransformSharedState::All;
+	EActorModifierTransformSharedState SaveState = EActorModifierTransformSharedState::None;
 };
 
 USTRUCT()
@@ -102,7 +102,7 @@ struct FActorModifierTransformSharedActorState
 
 	/** Used to restore only what has changed */
 	UPROPERTY()
-	EActorModifierTransformSharedState SaveState = EActorModifierTransformSharedState::All;
+	EActorModifierTransformSharedState SaveState = EActorModifierTransformSharedState::None;
 };
 
 /**
diff --git a/Engine/Plugins/Experimental/ActorModifier/Source/ActorModifierLayout/Private/Modifiers/ActorModifierAlignBetweenModifier.cpp b/Engine/Plugins/Experimental/ActorModifier/Source/ActorModifierLayout/Private/Modifiers/ActorModifierAlignBetweenModifier.cpp
--- a/Engine/Plugins/Experimental/ActorModifier/Source/ActorModifierLayout/Private/Modifiers/ActorModifierAlignBetweenModifier.cpp
+++ b/Engine/Plugins/Experimental/ActorModifier/Source/ActorModifierLayout/Private/Modifiers/ActorModifierAlignBetweenModifier.cpp
@@ -22,7 +22,7 @@ void UActorModifierAlignBetweenModifier::OnModifierEnabled(EActorModifierCoreEna
 	// Save actor layout state
 	if (UActorModifierTransformShared* LayoutShared = GetShared<UActorModifierTransformShared>(true))
 	{
-		LayoutShared->SaveActorState(this, GetModifiedActor());
+		LayoutShared->SaveActorState(this, GetModifiedActor(), EActorModifierTransformSharedState::Location);
 	}
 
 	SetTransformExtensionReferenceActors();
@@ -35,7 +35,7 @@ void UActorModifierAlignBetweenModifier::OnModifierDisabled(EActorModifierCoreDi
 	// Restore actor layout state
 	if (UActorModifierTransformShared* LayoutShared = GetShared<UActorModifierTransformShared>(false))
 	{
-		LayoutShared->RestoreActorState(this, GetModifiedActor());
+		LayoutShared->RestoreActorState(this, GetModifiedActor(), EActorModifierTransformSharedState::Location);
 	}
 }
 
diff --git a/Engine/Plugins/Experimental/ActorModifierCore/Source/ActorModifierCore/Private/Modifiers/ActorModifierCoreBase.cpp b/Engine/Plugins/Experimental/ActorModifierCore/Source/ActorModifierCore/Private/Modifiers/ActorModifierCoreBase.cpp
--- a/Engine/Plugins/Experimental/ActorModifierCore/Source/ActorModifierCore/Private/Modifiers/ActorModifierCoreBase.cpp
+++ b/Engine/Plugins/Experimental/ActorModifierCore/Source/ActorModifierCore/Private/Modifiers/ActorModifierCoreBase.cpp
@@ -33,76 +33,6 @@ UActorModifierCoreSharedObject* UActorModifierCoreBase::GetShared(TSubclassOf<UA
 	return nullptr;
 }
 
-TFuture<bool> UActorModifierCoreBase::ExecuteModifier()
-{
-	ExecutePromise = MakeShared<TPromise<bool>>();
-	TFuture<bool> ExecuteFuture = ExecutePromise->GetFuture();
-
-	const bool bValidActor = IsValid(GetModifiedActor());
-	const bool bModifierReady = IsModifierReady();
-
-	if (bModifierIdle && bValidActor && bModifierReady)
-	{
-		if (bModifierEnabled || IsModifierStack())
-		{
-			// lock current execution state
-			BeginModifierExecution();
-
-			// save the state before this modifier is execute to restore it later
-			SavePreState();
-
-			// run modifier logic, and update the modifier dirty state once the modifier logic completes
-			TWeakObjectPtr<UActorModifierCoreBase> WeakThis(this);
-			ExecuteApply()
-			.Then([WeakThis](TFuture<bool> ApplyResult)
-			{
-				UActorModifierCoreBase* This = WeakThis.Get();
-				if (!This)
-				{
-					return;
-				}
-
-				const bool bApplyResult = ApplyResult.Get();
-
-				// If it succeeded then unmark modifier dirty and set state as applied (for restore)
-				if (bApplyResult)
-				{
-					This->bModifierDirty = false;
-					This->bModifierApplied = true;
-				}
-
-				This->ApplyPromise.Reset();
-
-				// unlock current execution state
-				This->EndModifierExecution();
-
-				// return result execution
-				This->ExecutePromise->SetValue(bApplyResult);
-				This->ExecutePromise.Reset();
-			});
-		}
-		// disable
-		else
-		{
-			// Skipped
-			Status = FActorModifierCoreStatus(EActorModifierCoreStatus::Success, FText::GetEmpty());
-
-			bModifierDirty = false;
-
-			ExecutePromise->SetValue(true);
-			ExecutePromise.Reset();
-		}
-	}
-	// invalid
-	else
-	{
-		ExecutePromise->SetValue(false);
-		ExecutePromise.Reset();
-	}
-
-	return ExecuteFuture;
-}
-
 void UActorModifierCoreBase::Apply()
 {
 	checkNoEntry();
@@ -113,16 +43,19 @@ void UActorModifierCoreBase::Next()
 {
 	auto ExecuteNext = [this]()
 	{
-		if (!bModifierIdle && ApplyPromise.IsValid())
+		if (UActorModifierCoreStack* ModifierStack = GetModifierStack())
 		{
-			// Success
-			Status = FActorModifierCoreStatus(EActorModifierCoreStatus::Success, FText::GetEmpty());
-			ApplyPromise->SetValue(true);
-		}
-		else
-		{
-			LogModifier(TEXT("Next is called again after execution is done"), true);
-			checkNoEntry()
+			if (!bModifierIdle && ModifierStack->ExecutionTask.GetCurrentModifier() == this)
+			{
+				// Success
+				Status = FActorModifierCoreStatus(EActorModifierCoreStatus::Success, FText::GetEmpty());
+				ModifierStack->ExecutionTask.Next();
+			}
+			else
+			{
+				LogModifier(TEXT("Next is called again after execution is done"), true);
+				checkNoEntry()
+			}
 		}
 	};
 
@@ -149,20 +82,45 @@ void UActorModifierCoreBase::Next()
 
 void UActorModifierCoreBase::Fail(const FText& InFailReason)
 {
-	if (!bModifierIdle && ApplyPromise.IsValid())
+	auto ExecuteFail = [this, InFailReason]()
+	{
+		if (UActorModifierCoreStack* ModifierStack = GetModifierStack())
+		{
+			if (!bModifierIdle && ModifierStack->ExecutionTask.GetCurrentModifier() == this)
+			{
+				// Provide a valid failing reason
+				checkf(!InFailReason.IsEmpty(), TEXT("Provide a valid fail reason"))
+
+				// Success
+				Status = FActorModifierCoreStatus(EActorModifierCoreStatus::Error, InFailReason);
+				ModifierStack->ExecutionTask.Fail();
+			}
+			else
+			{
+				LogModifier(TEXT("Fail is called again after execution is done"), true);
+				checkNoEntry()
+			}
+		}
+	};
+	
+	if (!IsInGameThread())
 	{
-		// Provide a valid failing reason
-		check(!InFailReason.IsEmpty())
+		TWeakObjectPtr<UActorModifierCoreBase> ThisWeak(this);
+		FFunctionGraphTask::CreateAndDispatchWhenReady([ThisWeak, ExecuteFail]()
+		{
+			const UActorModifierCoreBase* This = ThisWeak.Get();
 
-		// Failed
-		Status = FActorModifierCoreStatus(EActorModifierCoreStatus::Error, InFailReason);
+			if (!This)
+			{
+				return;
+			}
 
-		ApplyPromise->SetValue(false);
+			ExecuteFail();
+		}, {}, nullptr, ENamedThreads::GameThread);
 	}
 	else
 	{
-		LogModifier(TEXT("Fail is called again after execution is done"), true);
-		checkNoEntry()
+		ExecuteFail();
 	}
 }
 
@@ -350,11 +308,18 @@ void UActorModifierCoreBase::MarkModifierDirty(bool bExecute)
 {
 	const UActorModifierCoreStack* Stack = GetRootModifierStack();
 
+	// Do not mark dirty if stack is in execution
 	if (!Stack || !Stack->IsModifierIdle() || !IsModifierInitialized())
 	{
 		return;
 	}
 
+	// When modifier is disabled but applied, we need to mark it dirty for it to be restored
+	if (!IsModifierEnabled() && !IsModifierApplied())
+	{
+		return;
+	}
+
 	if (!bModifierDirty || bExecute)
 	{
 		bModifierDirty = true;
@@ -842,33 +807,10 @@ void UActorModifierCoreBase::UninitializeModifier(EActorModifierCoreDisableReaso
 		// recover old enabled state
 		bModifierEnabled = bWasModifierEnabled;
 
-		// Ensure promises are fulfilled
-		if (InReason == EActorModifierCoreDisableReason::Destroyed)
-		{
-			if (ApplyPromise.IsValid())
-			{
-				ApplyPromise->SetValue(false);
-				ApplyPromise.Reset();
-			}
-
-			if (ExecutePromise.IsValid())
-			{
-				ExecutePromise->SetValue(false);
-				ExecutePromise.Reset();
-			}
-		}
-
 		UActorModifierCoreStack::OnModifierRemovedDelegate.Broadcast(this, InReason);
 	}
 }
 
-TFuture<bool> UActorModifierCoreBase::ExecuteApply()
-{
-	ApplyPromise = MakeShared<TPromise<bool>>();
-	Apply();
-	return ApplyPromise->GetFuture();
-}
-
 void UActorModifierCoreBase::OnModifierEnabledChanged(bool bInExecute)
 {
 	LogModifier(FString::Printf(TEXT("Modifier %s"), bModifierEnabled ? TEXT("enabled") : TEXT("disabled")), true);
diff --git a/Engine/Plugins/Experimental/ActorModifierCore/Source/ActorModifierCore/Private/Modifiers/ActorModifierCoreDefs.cpp b/Engine/Plugins/Experimental/ActorModifierCore/Source/ActorModifierCore/Private/Modifiers/ActorModifierCoreDefs.cpp
--- a/Engine/Plugins/Experimental/ActorModifierCore/Source/ActorModifierCore/Private/Modifiers/ActorModifierCoreDefs.cpp
+++ b/Engine/Plugins/Experimental/ActorModifierCore/Source/ActorModifierCore/Private/Modifiers/ActorModifierCoreDefs.cpp
@@ -1,6 +1,7 @@
 // Copyright Epic Games, Inc. All Rights Reserved.
 
 #include "Modifiers/ActorModifierCoreDefs.h"
+
 #include "Modifiers/ActorModifierCoreBase.h"
 #include "Modifiers/ActorModifierCoreStack.h"
 #include "Subsystems/ActorModifierCoreSubsystem.h"
@@ -321,3 +322,216 @@ FActorModifierCoreScopedLock::~FActorModifierCoreScopedLock()
 
 	ModifiersWeak.Empty();
 }
+
+void FActorModifierCoreExecutionTask::Restore()
+{
+	bool bAllModifiersDirty = false;
+	if (ModifierStack)
+	{
+		bAllModifiersDirty = ModifierStack->bAllModifiersDirty;
+	}
+
+	// build the restore chain based on dirty modifiers
+	bool bModifierBeforeDirty = false;
+	TArray<TObjectPtr<UActorModifierCoreBase>> ModifiersRestoreChain;
+	ModifiersRestoreChain.Reserve(Modifiers.Num());
+	for (const TObjectPtr<UActorModifierCoreBase>& Modifier : Modifiers)
+	{
+		if (Modifier)
+		{
+			if ((bAllModifiersDirty || bModifierBeforeDirty || Modifier->IsModifierDirty()))
+			{
+				bModifierBeforeDirty = true;
+				ModifiersRestoreChain.Add(Modifier);
+			}
+		}
+	}
+
+	// unapply modifiers in reverse order
+	Algo::Reverse(ModifiersRestoreChain);
+
+	// unapply modifiers
+	for (TObjectPtr<UActorModifierCoreBase>& Modifier : ModifiersRestoreChain)
+	{
+		if (Modifier && Modifier->IsModifierApplied())
+		{
+			// will only work for modifier that were already executed
+			Modifier->Unapply();
+		}
+	}
+
+	CurrentIndex = INDEX_NONE;
+	Modifiers.Empty();
+}
+
+void FActorModifierCoreExecutionTask::Apply(TConstArrayView<TObjectPtr<UActorModifierCoreBase>> InModifiers, UActorModifierCoreStack* InStack)
+{
+	check(Modifiers.IsEmpty() && CurrentIndex == INDEX_NONE)
+
+	BuildExecutionChain(InModifiers, InStack);
+
+	if (Modifiers.IsValidIndex(CurrentIndex))
+	{
+		ExecuteCurrentModifier();
+	}
+	else if (ModifierStack)
+	{
+		ModifierStack->OnModifierExecutionFinished(/** Result */true);
+	}
+}
+
+void FActorModifierCoreExecutionTask::Next()
+{
+	UActorModifierCoreBase* CurrentModifier = Modifiers[CurrentIndex];
+
+	if (!CurrentModifier->IsModifierIdle())
+	{
+		// unlock current execution state
+		CurrentModifier->EndModifierExecution();
+	}
+
+	// If it succeeded then unmark modifier dirty and set state as applied (for restore)
+	CurrentModifier->bModifierDirty = false;
+	CurrentModifier->bModifierApplied = true;
+
+	// Next
+	CurrentIndex++;
+	ExecuteCurrentModifier();
+}
+
+void FActorModifierCoreExecutionTask::Fail()
+{
+	UActorModifierCoreBase* CurrentModifier = Modifiers[CurrentIndex];
+
+	if (!CurrentModifier->IsModifierIdle())
+	{
+		// unlock current execution state
+		CurrentModifier->EndModifierExecution();
+	}
+	
+	if (ModifierStack)
+	{
+		ModifierStack->OnModifierExecutionFinished(/** Result */false);
+	}
+}
+
+void FActorModifierCoreExecutionTask::Skip()
+{
+	UActorModifierCoreBase* CurrentModifier = Modifiers[CurrentIndex];
+
+	// Skipped
+	CurrentModifier->Status = FActorModifierCoreStatus(EActorModifierCoreStatus::Success, FText::GetEmpty());
+	CurrentModifier->bModifierDirty = false;
+	CurrentModifier->bModifierApplied = false;
+
+	// Next
+	CurrentIndex++;
+	ExecuteCurrentModifier();
+}
+
+UActorModifierCoreBase* FActorModifierCoreExecutionTask::GetCurrentModifier() const
+{
+	if (Modifiers.IsValidIndex(CurrentIndex))
+	{
+		return Modifiers[CurrentIndex];
+	}
+
+	return nullptr;
+}
+
+void FActorModifierCoreExecutionTask::ReplaceModifier(UActorModifierCoreBase& InOldModifier, UActorModifierCoreBase& InNewModifier)
+{
+	const int32 Index = Modifiers.Find(&InOldModifier);
+	
+	if (Modifiers.IsValidIndex(Index))
+	{
+		Modifiers[Index] = &InNewModifier;
+	}
+}
+
+void FActorModifierCoreExecutionTask::RemoveModifier(UActorModifierCoreBase& InModifier)
+{
+	Modifiers.Remove(&InModifier);
+}
+
+void FActorModifierCoreExecutionTask::BuildExecutionChain(TConstArrayView<TObjectPtr<UActorModifierCoreBase>> InModifiers, UActorModifierCoreStack* InStack)
+{
+	ModifierStack = InStack;
+	Modifiers = InModifiers;
+
+	bool bAllModifiersDirty = false;
+
+	if (ModifierStack)
+	{
+		bAllModifiersDirty = ModifierStack->bAllModifiersDirty;
+	}
+
+	for (int32 Index = 0; Index < Modifiers.Num(); ++Index)
+	{
+		if (UActorModifierCoreBase* Modifier = Modifiers[Index])
+		{
+			if (bAllModifiersDirty
+				|| Modifier->IsModifierDirty()
+				// Give a chance to non tickable modifier to set themselves dirty too
+				|| (!Modifier->GetModifierMetadata().IsTickAllowed() && Modifier->IsModifierDirtyable()))
+			{
+				CurrentIndex = Index;
+				break;
+			}
+		}
+	}
+}
+
+void FActorModifierCoreExecutionTask::ExecuteCurrentModifier()
+{
+	// are we done with this execution round
+	if (CurrentIndex > Modifiers.Num() - 1)
+	{
+		if (ModifierStack)
+		{
+			ModifierStack->OnModifierExecutionFinished(/** Result */true);
+		}
+
+		return;
+	}
+
+	const TObjectPtr<UActorModifierCoreBase>& Modifier = Modifiers[CurrentIndex];
+
+	if (!IsValid(Modifier))
+	{
+		if (ModifierStack)
+		{
+			ModifierStack->OnModifierExecutionFinished(/** Result */false);
+		}
+
+		return;
+	}
+
+	const bool bValidActor = IsValid(Modifier->GetModifiedActor());
+	const bool bModifierReady = Modifier->bModifierIdle && Modifier->IsModifierReady();
+
+	if (bValidActor && bModifierReady)
+	{
+		if (Modifier->bModifierEnabled || Modifier->IsModifierStack())
+		{
+			// lock current execution state
+			Modifier->BeginModifierExecution();
+
+			// save the state before this modifier is execute to restore it later
+			Modifier->SavePreState();
+
+			// run modifier logic, and update the modifier dirty state once the modifier logic completes
+			Modifier->Apply();
+		}
+		// disable
+		else
+		{
+			Skip();
+		}
+	}
+	// invalid
+	else
+	{
+		Fail();
+	}
+}
diff --git a/Engine/Plugins/Experimental/ActorModifierCore/Source/ActorModifierCore/Private/Modifiers/ActorModifierCoreStack.cpp b/Engine/Plugins/Experimental/ActorModifierCore/Source/ActorModifierCore/Private/Modifiers/ActorModifierCoreStack.cpp
--- a/Engine/Plugins/Experimental/ActorModifierCore/Source/ActorModifierCore/Private/Modifiers/ActorModifierCoreStack.cpp
+++ b/Engine/Plugins/Experimental/ActorModifierCore/Source/ActorModifierCore/Private/Modifiers/ActorModifierCoreStack.cpp
@@ -1237,8 +1237,7 @@ bool UActorModifierCoreStack::RemoveModifier(FActorModifierCoreStackRemoveOp& In
 			{
 				RemoveModifier->UninitializeModifier(EActorModifierCoreDisableReason::User);
 
-				This->ExecuteModifiers.Remove(RemoveModifier);
-				This->CurrentModifiers.Remove(RemoveModifier);
+				This->ExecutionTask.RemoveModifier(*RemoveModifier);
 
 				RemoveModifier->MarkAsGarbage();
 			}
@@ -1397,75 +1396,43 @@ void UActorModifierCoreStack::OnActorTransformUpdated(USceneComponent*, EUpdateT
 	}
 }
 
-void UActorModifierCoreStack::ApplyChainModifiers()
+void UActorModifierCoreStack::OnModifierExecutionFinished(bool bInResult)
 {
-	// are we done with this execution round
-	if (ModifierExecutionIdx > ExecuteModifiers.Num()-1)
-	{
-		bAllModifiersDirty = false;
-		Next();
-		return;
-	}
+	bAllModifiersDirty = false;
 
-	// execute current modifier idx
-	if (ExecuteModifiers.IsValidIndex(ModifierExecutionIdx))
+	if (bInResult)
 	{
-		const TObjectPtr<UActorModifierCoreBase>& Modifier = ExecuteModifiers[ModifierExecutionIdx];
-
-		TWeakObjectPtr<UActorModifierCoreStack> ThisWeak(this);
-		Modifier->ExecuteModifier()
-		.Then([ThisWeak](TFuture<bool> ExecuteResult)
-		{
-			if (!ThisWeak.IsValid())
-			{
-				return;
-			}
-
-			// how did the execution go
-			const bool bResult = ExecuteResult.Get();
-			if (!bResult)
-			{
-				ThisWeak->bAllModifiersDirty = false;
-				ThisWeak->Fail(LOCTEXT("ModifierStackExecutionFail", "A modifier in this stack failed to execute"));
-			}
-			else
-			{
-				ThisWeak->ModifierExecutionIdx++;
-				ThisWeak->ApplyChainModifiers();
-			}
-		});
+		Next();
 	}
-	// invalid modifier found
 	else
 	{
-		bAllModifiersDirty = false;
-		Fail(LOCTEXT("InvalidModifierIndexFail", "An invalid modifier index provided for the execution"));
+		Fail(LOCTEXT("ModifierStackExecutionFail", "A modifier in this stack failed to execute"));
 	}
-}
 
-void UActorModifierCoreStack::BuildExecutionChain()
-{
-	ExecuteModifiers.Empty(Modifiers.Num());
-	ModifierExecutionIdx = 0;
-	bool bModifierBeforeDirty = false;
-	for (const TObjectPtr<UActorModifierCoreBase>& Modifier : Modifiers)
+	if (IsRootStack() && RootExecutionTask.IsSet())
 	{
-		if (Modifier)
+		if (bInResult)
 		{
-			if (bAllModifiersDirty
-				|| bModifierBeforeDirty
-				|| Modifier->IsModifierDirty()
-				// Give a chance to non tickable modifier to set themselves dirty too
-				|| (!Modifier->GetModifierMetadata().IsTickAllowed() && Modifier->IsModifierDirtyable()))
-			{
-				bModifierBeforeDirty = true;
+			RootExecutionTask->Next();
+		}
+		else
+		{
+			RootExecutionTask->Fail();
+		}
+		
+		RootExecutionTask.Reset();
 
-				// only add if enabled
-				if (Modifier->IsModifierEnabled())
-				{
-					ExecuteModifiers.Add(Modifier);
-				}
+		// process idle function
+		if (!IsModifierExecutionLocked())
+		{
+			ScheduleModifierOptimization(false);
+
+			for (int32 Index = 0; Index < OnIdleFunctions.Num(); ++Index)
+			{
+				OnIdleFunctions[Index]();
 			}
+
+			OnIdleFunctions.Empty();
 		}
 	}
 }
@@ -1628,19 +1595,7 @@ void UActorModifierCoreStack::OnBlueprintModifierReplaced(UActorModifierCoreBlue
 	}
 
 	// Replace old modifier by new modifier in stack
-	const int32 ExecuteIndex = ExecuteModifiers.Find(InOldModifier);
-
-	if (ExecuteModifiers.IsValidIndex(ExecuteIndex))
-	{
-		ExecuteModifiers[ExecuteIndex] = InNewModifier;
-	}
-
-	const int32 CurrentIndex = CurrentModifiers.Find(InOldModifier);
-
-	if (CurrentModifiers.IsValidIndex(CurrentIndex))
-	{
-		CurrentModifiers[CurrentIndex] = InNewModifier;
-	}
+	ExecutionTask.ReplaceModifier(*InOldModifier, *InNewModifier);
 
 	InNewModifier->OnModifierReplacedEvent(TargetActor);
 
@@ -1651,33 +1606,8 @@ void UActorModifierCoreStack::OnBlueprintModifierReplaced(UActorModifierCoreBlue
 
 void UActorModifierCoreStack::RestorePreState()
 {
-	TArray<TObjectPtr<UActorModifierCoreBase>> ModifiersRestoreChain;
-
-	// build the restore chain based on dirty modifiers
-	bool bModifierBeforeDirty = false;
-	for (const TObjectPtr<UActorModifierCoreBase>& Modifier : CurrentModifiers)
-	{
-		if (Modifier)
-		{
-			if ((bAllModifiersDirty || bModifierBeforeDirty || Modifier->IsModifierDirty()))
-			{
-				bModifierBeforeDirty = true;
-				ModifiersRestoreChain.Add(Modifier);
-			}
-		}
-	}
-	// unapply modifiers in reverse order
-	Algo::Reverse(ModifiersRestoreChain);
-
-	// unapply modifiers
-	for (TObjectPtr<UActorModifierCoreBase>& Modifier : ModifiersRestoreChain)
-	{
-		if (Modifier)
-		{
-			// will only work for modifier that were already executed
-			Modifier->Unapply();
-		}
-	}
+	// Restore previously applied modifiers
+	ExecutionTask.Restore();
 
 	// Stack is restored, execute one time callbacks
 	for (int32 Index = 0; Index < OnRestoreFunctions.Num(); ++Index)
@@ -1685,8 +1615,6 @@ void UActorModifierCoreStack::RestorePreState()
 		OnRestoreFunctions[Index]();
 	}
 	OnRestoreFunctions.Empty();
-
-	CurrentModifiers.Empty(Modifiers.Num());
 }
 
 void UActorModifierCoreStack::Apply()
@@ -1694,14 +1622,7 @@ void UActorModifierCoreStack::Apply()
 	// restore original state before dirty modifiers
 	RestorePreState();
 
-	// build the execution chain of dirty modifiers to execute this round
-	BuildExecutionChain();
-
-	// recurse apply them
-	ApplyChainModifiers();
-
-	// copy to restore later in case the original array is altered (remove, add, move)
-	CurrentModifiers = Modifiers;
+	ExecutionTask.Apply(Modifiers, this);
 }
 
 void UActorModifierCoreStack::OnModifierRemoved(EActorModifierCoreDisableReason InReason)
@@ -1795,33 +1716,9 @@ void UActorModifierCoreStack::OnModifierDirty(UActorModifierCoreBase* DirtyModif
 
 	LogModifier(TEXT("Executing modifier root stack"));
 
-	TWeakObjectPtr<UActorModifierCoreStack> WeakThis(this);
-
-	ExecuteModifier()
-	.Then([WeakThis](TFuture<bool> ExecuteResult)
-	{
-		UActorModifierCoreStack* This = WeakThis.Get();
-		if (!IsValid(This))
-		{
-			return;
-		}
-
-		const bool bResult = ExecuteResult.Get();
-		This->ExecutePromise.Reset();
-
-		// process idle function
-		if (!This->IsModifierExecutionLocked())
-		{
-			This->ScheduleModifierOptimization(false);
-
-			for (int32 Index = 0; Index < This->OnIdleFunctions.Num(); ++Index)
-			{
-				This->OnIdleFunctions[Index]();
-			}
-
-			This->OnIdleFunctions.Empty();
-		}
-	});
+	// Execute root stack as standalone modifier
+	RootExecutionTask = FActorModifierCoreExecutionTask();
+	RootExecutionTask->Apply({this}, nullptr);
 }
 
 void UActorModifierCoreStack::OnModifierCDOSetup(FActorModifierCoreMetadata& InMetadata)
diff --git a/Engine/Plugins/Experimental/ActorModifierCore/Source/ActorModifierCore/Public/Modifiers/ActorModifierCoreBase.h b/Engine/Plugins/Experimental/ActorModifierCore/Source/ActorModifierCore/Public/Modifiers/ActorModifierCoreBase.h
--- a/Engine/Plugins/Experimental/ActorModifierCore/Source/ActorModifierCore/Public/Modifiers/ActorModifierCoreBase.h
+++ b/Engine/Plugins/Experimental/ActorModifierCore/Source/ActorModifierCore/Public/Modifiers/ActorModifierCoreBase.h
@@ -29,11 +29,12 @@ class UActorModifierCoreBase : public UObject
 	friend class UActorModifierCoreStack;
 	friend class UActorModifierCoreSubsystem;
 
-	friend struct FActorModifierCoreMetadata;
-	friend struct FActorModifierCoreScopedLock;
-
 	friend class FActorModifierCoreEditorDetailCustomization;
 	friend class UActorModifierCoreEditorStackCustomization;
+	
+	friend struct FActorModifierCoreExecutionTask;
+	friend struct FActorModifierCoreMetadata;
+	friend struct FActorModifierCoreScopedLock;
 
 	GENERATED_BODY()
 
@@ -395,12 +396,6 @@ private:
 	/** INTERNAL USE ONLY, add and setup a newly created extension for this modifier */
 	ACTORMODIFIERCORE_API void AddExtensionInternal(const FName& InExtensionType, TSharedPtr<FActorModifierCoreExtension> InExtension);
 
-	/** Execute this modifier and calls the appropriate functions (SavePreState->Apply->SavePostState) when needed */
-	TFuture<bool> ExecuteModifier();
-
-	/** Calls apply and handle the async feature without having to deal with it in the children classes */
-	TFuture<bool> ExecuteApply();
-
 	/** Called after setting the modifier enabled state */
 	void OnModifierEnabledChanged(bool bInExecute = true);
 
@@ -415,12 +410,6 @@ private:
 	/** Map of typed extensions, will follow this modifier lifecycle */
 	TMap<FName, TSharedPtr<FActorModifierCoreExtension, ESPMode::ThreadSafe>> ModifierExtensions;
 
-	/** Promise that executes when the modifier apply is done */
-	TSharedPtr<TPromise<bool>, ESPMode::ThreadSafe> ApplyPromise = nullptr;
-
-	/** Promise that executes when the modifier execution is done */
-	TSharedPtr<TPromise<bool>, ESPMode::ThreadSafe> ExecutePromise = nullptr;
-
 	TWeakObjectPtr<AActor> ModifiedActor = nullptr;
 
 	/** Is the modifier enabled or disabled */
diff --git a/Engine/Plugins/Experimental/ActorModifierCore/Source/ActorModifierCore/Public/Modifiers/ActorModifierCoreDefs.h b/Engine/Plugins/Experimental/ActorModifierCore/Source/ActorModifierCore/Public/Modifiers/ActorModifierCoreDefs.h
--- a/Engine/Plugins/Experimental/ActorModifierCore/Source/ActorModifierCore/Public/Modifiers/ActorModifierCoreDefs.h
+++ b/Engine/Plugins/Experimental/ActorModifierCore/Source/ActorModifierCore/Public/Modifiers/ActorModifierCoreDefs.h
@@ -448,3 +448,36 @@ protected:
 	/** Modifier last status message */
 	FText StatusMessage = FText::GetEmpty();
 };
+
+USTRUCT()
+struct FActorModifierCoreExecutionTask
+{
+	GENERATED_BODY()
+
+	/** Restore original state until first dirty modifier, last -> first */
+	void Restore();
+	
+	/** Build the execution chain of dirty modifiers and execute, first -> last */
+	void Apply(TConstArrayView<TObjectPtr<UActorModifierCoreBase>> InModifiers, UActorModifierCoreStack* InStack);
+
+	void Next();
+	void Fail();
+	void Skip();
+
+	UActorModifierCoreBase* GetCurrentModifier() const;
+
+	void ReplaceModifier(UActorModifierCoreBase& InOldModifier, UActorModifierCoreBase& InNewModifier);
+	void RemoveModifier(UActorModifierCoreBase& InModifier);
+
+private:
+	void BuildExecutionChain(TConstArrayView<TObjectPtr<UActorModifierCoreBase>> InModifiers, UActorModifierCoreStack* InStack);
+	void ExecuteCurrentModifier();
+
+	UPROPERTY()
+	TObjectPtr<UActorModifierCoreStack> ModifierStack;
+
+	UPROPERTY()
+	TArray<TObjectPtr<UActorModifierCoreBase>> Modifiers;
+
+	int32 CurrentIndex = INDEX_NONE;
+};
\ No newline at end of file
diff --git a/Engine/Plugins/Experimental/ActorModifierCore/Source/ActorModifierCore/Public/Modifiers/ActorModifierCoreStack.h b/Engine/Plugins/Experimental/ActorModifierCore/Source/ActorModifierCore/Public/Modifiers/ActorModifierCoreStack.h
--- a/Engine/Plugins/Experimental/ActorModifierCore/Source/ActorModifierCore/Public/Modifiers/ActorModifierCoreStack.h
+++ b/Engine/Plugins/Experimental/ActorModifierCore/Source/ActorModifierCore/Public/Modifiers/ActorModifierCoreStack.h
@@ -24,6 +24,8 @@ class UActorModifierCoreStack : public UActorModifierCoreBase
 	friend class FActorModifierCoreEditorDetailCustomization;
 	friend class UActorModifierCoreEditorStackCustomization;
 
+	friend struct FActorModifierCoreExecutionTask;
+
 public:
 	DECLARE_MULTICAST_DELEGATE_OneParam(FOnModifierUpdated, UActorModifierCoreBase* /** UpdatedItem */)
 	DECLARE_MULTICAST_DELEGATE_TwoParams(FOnModifierAdded, UActorModifierCoreBase* /** Item */, EActorModifierCoreEnableReason /** Reason */)
@@ -220,11 +222,8 @@ private:
 	/** Called when the actor transform is updated */
 	void OnActorTransformUpdated(USceneComponent*, EUpdateTransformFlags, ETeleportType);
 
-	/** Will execute the chain of modifier in this stack */
-	void ApplyChainModifiers();
-
-	/** Builds the execution chain of modifiers */
-	void BuildExecutionChain();
+	/** Called when the execution task is done */
+	void OnModifierExecutionFinished(bool bInResult);
 
 	/** Schedules a modifier optimization on idle */
 	void ScheduleModifierOptimization(bool bInInvalidateAll = false);
@@ -239,25 +238,22 @@ private:
 	UPROPERTY(VisibleInstanceOnly, NoClear, Export, Instanced, Category = "Modifiers")
 	TArray<TObjectPtr<UActorModifierCoreBase>> Modifiers;
 
-	/** Contains a copy of modifiers in the stack for this round of execution, useful for restore and for query, can be different from modifiers array */
 	UPROPERTY(Transient, DuplicateTransient, NonTransactional)
-	TArray<TObjectPtr<UActorModifierCoreBase>> CurrentModifiers;
+	FActorModifierCoreExecutionTask ExecutionTask;
 
-	/** Contains the modifiers that we will execute this round, can be different from modifiers array */
+	/** Used by root stack to trigger itself */
 	UPROPERTY(Transient, DuplicateTransient, NonTransactional)
-	TArray<TObjectPtr<UActorModifierCoreBase>> ExecuteModifiers;
-
-	/** Functions to execute once when the stack is on idle */
-	TArray<TFunction<void()>> OnIdleFunctions;
-
-	/** Functions to execute once when the stack is restored */
-	TArray<TFunction<void()>> OnRestoreFunctions;
+	TOptional<FActorModifierCoreExecutionTask> RootExecutionTask;
 
 	/** Enable profiling for the modifiers in this stack */
 	UPROPERTY(EditInstanceOnly, Category="Modifier", meta=(DisplayName="Enable Profiling"))
 	bool bModifierProfiling = false;
 
-	int32 ModifierExecutionIdx = INDEX_NONE;
-
 	bool bAllModifiersDirty = true;
+
+	/** Functions to execute once when the stack is on idle */
+	TArray<TFunction<void()>> OnIdleFunctions;
+
+	/** Functions to execute once when the stack is restored */
+	TArray<TFunction<void()>> OnRestoreFunctions;
 };
diff --git a/Engine/Plugins/Experimental/Avalanche/Source/AvalancheModifiers/Private/Modifiers/AvaBooleanModifier.cpp b/Engine/Plugins/Experimental/Avalanche/Source/AvalancheModifiers/Private/Modifiers/AvaBooleanModifier.cpp
--- a/Engine/Plugins/Experimental/Avalanche/Source/AvalancheModifiers/Private/Modifiers/AvaBooleanModifier.cpp
+++ b/Engine/Plugins/Experimental/Avalanche/Source/AvalancheModifiers/Private/Modifiers/AvaBooleanModifier.cpp
@@ -90,22 +90,10 @@ void UAvaBooleanModifier::Apply()
 	// Update masking material here too to avoid flickering between mask and original materials due to async task
 	UpdateMaskingMaterials();
 
-	// Use async execution since boolean is an heavy operation
-	TWeakObjectPtr<UAvaBooleanModifier> ThisWeak(this);
-	Async(EAsyncExecution::TaskGraphMainThread, [ThisWeak]()
-	{
-		UAvaBooleanModifier* This = ThisWeak.Get();
-
-		if (!IsValid(This))
-		{
-			return;
-		}
-
-		This->ApplyInternal();
-		This->UpdateMaskingMaterials();
+	ApplyInternal();
+	UpdateMaskingMaterials();
 
-		This->Next();
-	});
+	Next();
 }
 
 void UAvaBooleanModifier::OnModifierDisabled(EActorModifierCoreDisableReason InReason)
