MotionDesign : Modifiers
- Set localization for modifiers in operator stack/menus

diff --git a/Engine/Plugins/VirtualProduction/ActorModifier/Source/ActorModifierLayout/Private/Modifiers/ActorModifierAlignBetweenModifier.cpp b/Engine/Plugins/VirtualProduction/ActorModifier/Source/ActorModifierLayout/Private/Modifiers/ActorModifierAlignBetweenModifier.cpp
--- a/Engine/Plugins/VirtualProduction/ActorModifier/Source/ActorModifierLayout/Private/Modifiers/ActorModifierAlignBetweenModifier.cpp
+++ b/Engine/Plugins/VirtualProduction/ActorModifier/Source/ActorModifierLayout/Private/Modifiers/ActorModifierAlignBetweenModifier.cpp
@@ -195,6 +195,7 @@ void UActorModifierAlignBetweenModifier::OnModifierCDOSetup(FActorModifierCoreMe
 	InMetadata.SetName(TEXT("AlignBetween"));
 	InMetadata.SetCategory(TEXT("Layout"));
 #if WITH_EDITOR
+	InMetadata.SetDisplayName(LOCTEXT("ModifierDisplayName", "Align Between"));
 	InMetadata.SetDescription(LOCTEXT("ModifierDescription", "Positions an actor between a group of weighted actors"));
 #endif
 }
diff --git a/Engine/Plugins/VirtualProduction/ActorModifier/Source/ActorModifierLayout/Private/Modifiers/ActorModifierAutoFollowModifier.cpp b/Engine/Plugins/VirtualProduction/ActorModifier/Source/ActorModifierLayout/Private/Modifiers/ActorModifierAutoFollowModifier.cpp
--- a/Engine/Plugins/VirtualProduction/ActorModifier/Source/ActorModifierLayout/Private/Modifiers/ActorModifierAutoFollowModifier.cpp
+++ b/Engine/Plugins/VirtualProduction/ActorModifier/Source/ActorModifierLayout/Private/Modifiers/ActorModifierAutoFollowModifier.cpp
@@ -40,6 +40,7 @@ void UActorModifierAutoFollowModifier::OnModifierCDOSetup(FActorModifierCoreMeta
 	InMetadata.SetName(TEXT("AutoFollow"));
 	InMetadata.SetCategory(TEXT("Layout"));
 #if WITH_EDITOR
+	InMetadata.SetDisplayName(LOCTEXT("ModifierDisplayName", "Auto Follow"));
 	InMetadata.SetDescription(LOCTEXT("ModifierDescription", "Positions an actor relative to another actor using their bounds"));
 #endif
 }
diff --git a/Engine/Plugins/VirtualProduction/ActorModifier/Source/ActorModifierLayout/Private/Modifiers/ActorModifierGridArrangeModifier.cpp b/Engine/Plugins/VirtualProduction/ActorModifier/Source/ActorModifierLayout/Private/Modifiers/ActorModifierGridArrangeModifier.cpp
--- a/Engine/Plugins/VirtualProduction/ActorModifier/Source/ActorModifierLayout/Private/Modifiers/ActorModifierGridArrangeModifier.cpp
+++ b/Engine/Plugins/VirtualProduction/ActorModifier/Source/ActorModifierLayout/Private/Modifiers/ActorModifierGridArrangeModifier.cpp
@@ -83,6 +83,7 @@ void UActorModifierGridArrangeModifier::OnModifierCDOSetup(FActorModifierCoreMet
 	InMetadata.SetName(TEXT("GridArrange"));
 	InMetadata.SetCategory(TEXT("Layout"));
 #if WITH_EDITOR
+	InMetadata.SetDisplayName(LOCTEXT("ModifierDisplayName", "Grid Arrange"));
 	InMetadata.SetDescription(LOCTEXT("ModifierDescription", "Positions child actors in a 2D grid format"));
 #endif
 }
diff --git a/Engine/Plugins/VirtualProduction/ActorModifier/Source/ActorModifierLayout/Private/Modifiers/ActorModifierJustifyModifier.cpp b/Engine/Plugins/VirtualProduction/ActorModifier/Source/ActorModifierLayout/Private/Modifiers/ActorModifierJustifyModifier.cpp
--- a/Engine/Plugins/VirtualProduction/ActorModifier/Source/ActorModifierLayout/Private/Modifiers/ActorModifierJustifyModifier.cpp
+++ b/Engine/Plugins/VirtualProduction/ActorModifier/Source/ActorModifierLayout/Private/Modifiers/ActorModifierJustifyModifier.cpp
@@ -18,6 +18,7 @@ void UActorModifierJustifyModifier::OnModifierCDOSetup(FActorModifierCoreMetadat
 	InMetadata.SetName(TEXT("Justify"));
 	InMetadata.SetCategory(TEXT("Layout"));
 #if WITH_EDITOR
+	InMetadata.SetDisplayName(LOCTEXT("ModifierDisplayName", "Justify"));
 	InMetadata.SetDescription(LOCTEXT("ModifierDescription", "Aligns child actors according to the specified justification option, based on their bounding boxes"));
 #endif
 }
diff --git a/Engine/Plugins/VirtualProduction/ActorModifier/Source/ActorModifierLayout/Private/Modifiers/ActorModifierLookAtModifier.cpp b/Engine/Plugins/VirtualProduction/ActorModifier/Source/ActorModifierLayout/Private/Modifiers/ActorModifierLookAtModifier.cpp
--- a/Engine/Plugins/VirtualProduction/ActorModifier/Source/ActorModifierLayout/Private/Modifiers/ActorModifierLookAtModifier.cpp
+++ b/Engine/Plugins/VirtualProduction/ActorModifier/Source/ActorModifierLayout/Private/Modifiers/ActorModifierLookAtModifier.cpp
@@ -179,6 +179,7 @@ void UActorModifierLookAtModifier::OnModifierCDOSetup(FActorModifierCoreMetadata
 	InMetadata.SetName(TEXT("LookAt"));
 	InMetadata.SetCategory(TEXT("Layout"));
 #if WITH_EDITOR
+	InMetadata.SetDisplayName(LOCTEXT("ModifierDisplayName", "Look At"));
 	InMetadata.SetDescription(LOCTEXT("ModifierDescription", "Rotates an actor to face another actor"));
 #endif
 }
diff --git a/Engine/Plugins/VirtualProduction/ActorModifier/Source/ActorModifierLayout/Private/Modifiers/ActorModifierRadialArrangeModifier.cpp b/Engine/Plugins/VirtualProduction/ActorModifier/Source/ActorModifierLayout/Private/Modifiers/ActorModifierRadialArrangeModifier.cpp
--- a/Engine/Plugins/VirtualProduction/ActorModifier/Source/ActorModifierLayout/Private/Modifiers/ActorModifierRadialArrangeModifier.cpp
+++ b/Engine/Plugins/VirtualProduction/ActorModifier/Source/ActorModifierLayout/Private/Modifiers/ActorModifierRadialArrangeModifier.cpp
@@ -193,6 +193,7 @@ void UActorModifierRadialArrangeModifier::OnModifierCDOSetup(FActorModifierCoreM
 	InMetadata.SetName(TEXT("RadialArrange"));
 	InMetadata.SetCategory(TEXT("Layout"));
 #if WITH_EDITOR
+	InMetadata.SetDisplayName(LOCTEXT("ModifierDisplayName", "Radial Arrange"));
 	InMetadata.SetDescription(LOCTEXT("ModifierDescription", "Positions child actors in a 2D radial format"));
 #endif
 }
diff --git a/Engine/Plugins/VirtualProduction/ActorModifier/Source/ActorModifierLayout/Private/Modifiers/ActorModifierSplinePathModifier.cpp b/Engine/Plugins/VirtualProduction/ActorModifier/Source/ActorModifierLayout/Private/Modifiers/ActorModifierSplinePathModifier.cpp
--- a/Engine/Plugins/VirtualProduction/ActorModifier/Source/ActorModifierLayout/Private/Modifiers/ActorModifierSplinePathModifier.cpp
+++ b/Engine/Plugins/VirtualProduction/ActorModifier/Source/ActorModifierLayout/Private/Modifiers/ActorModifierSplinePathModifier.cpp
@@ -166,6 +166,7 @@ void UActorModifierSplinePathModifier::OnModifierCDOSetup(FActorModifierCoreMeta
 	InMetadata.SetName(TEXT("SplinePath"));
 	InMetadata.SetCategory(TEXT("Layout"));
 #if WITH_EDITOR
+	InMetadata.SetDisplayName(LOCTEXT("ModifierDisplayName", "Spline Path"));
 	InMetadata.SetDescription(LOCTEXT("ModifierDescription", "Samples a spline for an actor to follow a path"));
 #endif
 
diff --git a/Engine/Plugins/VirtualProduction/ActorModifierCore/Source/ActorModifierCore/Private/Modifiers/ActorModifierCoreDefs.cpp b/Engine/Plugins/VirtualProduction/ActorModifierCore/Source/ActorModifierCore/Private/Modifiers/ActorModifierCoreDefs.cpp
--- a/Engine/Plugins/VirtualProduction/ActorModifierCore/Source/ActorModifierCore/Private/Modifiers/ActorModifierCoreDefs.cpp
+++ b/Engine/Plugins/VirtualProduction/ActorModifierCore/Source/ActorModifierCore/Private/Modifiers/ActorModifierCoreDefs.cpp
@@ -4,12 +4,15 @@
 
 #include "Modifiers/ActorModifierCoreBase.h"
 #include "Modifiers/ActorModifierCoreStack.h"
+#include "Modifiers/Blueprints/ActorModifierCoreBlueprintBase.h"
 #include "Subsystems/ActorModifierCoreSubsystem.h"
 
 #if WITH_EDITOR
 #include "Styling/SlateIconFinder.h"
 #endif
 
+#define LOCTEXT_NAMESPACE "ActorModifierCoreMetadata"
+
 FActorModifierCoreMetadata::FActorModifierCoreMetadata()
 	: bIsStack(false)
 #if WITH_EDITOR
@@ -27,10 +30,6 @@ FActorModifierCoreMetadata::FActorModifierCoreMetadata(const UActorModifierCoreB
 	bIsStack = InModifier->IsA<UActorModifierCoreStack>();
 	Category = DefaultCategory;
 
-#if WITH_EDITOR
-	bHidden = false;
-#endif
-
 	SetProfilerClass<FActorModifierCoreProfiler>();
 	SetCompatibilityRule([](const AActor* InActor)->bool
 	{
@@ -43,6 +42,10 @@ FActorModifierCoreMetadata::FActorModifierCoreMetadata(const UActorModifierCoreB
 	{
 		Icon = FSlateIconFinder::FindIconForClass(UActorModifierCoreBase::StaticClass());
 	}
+
+	Description = Class->GetToolTipText();
+	DisplayName = Class->GetDisplayNameText();
+	bHidden = Class->HasMetaData("Hidden");
 #endif
 }
 
@@ -152,6 +155,15 @@ bool FActorModifierCoreMetadata::ResetDefault()
 FActorModifierCoreMetadata& FActorModifierCoreMetadata::SetDisplayName(const FText& InName)
 {
 	DisplayName = InName;
+
+	if (const TSubclassOf<UActorModifierCoreBase> ModifierClass = GetClass())
+	{
+		if (ModifierClass->IsChildOf<UActorModifierCoreBlueprintBase>())
+		{
+			DisplayName = FText::Format(LOCTEXT("BlueprintModifierDisplayName", "{0} {1}"), DisplayName, LOCTEXT("BlueprintModifierLabel", "(Blueprint)"));
+		}
+	}
+
 	return *this;
 }
 
@@ -537,3 +549,5 @@ void FActorModifierCoreExecutionTask::ExecuteCurrentModifier()
 		Fail();
 	}
 }
+
+#undef LOCTEXT_NAMESPACE
diff --git a/Engine/Plugins/VirtualProduction/ActorModifierCore/Source/ActorModifierCoreEditor/Private/Modifiers/ActorModifierCoreEditorMenu.cpp b/Engine/Plugins/VirtualProduction/ActorModifierCore/Source/ActorModifierCoreEditor/Private/Modifiers/ActorModifierCoreEditorMenu.cpp
--- a/Engine/Plugins/VirtualProduction/ActorModifierCore/Source/ActorModifierCoreEditor/Private/Modifiers/ActorModifierCoreEditorMenu.cpp
+++ b/Engine/Plugins/VirtualProduction/ActorModifierCore/Source/ActorModifierCoreEditor/Private/Modifiers/ActorModifierCoreEditorMenu.cpp
@@ -15,8 +15,11 @@
 void UE::ActorModifierCoreEditor::OnExtendAddModifierMenu(UToolMenu* InAddToolMenu, const FActorModifierCoreEditorMenuData& InData)
 {
 	const UActorModifierCoreSubsystem* ModifierSubsystem = UActorModifierCoreSubsystem::Get();
+	const UActorModifierCoreEditorSubsystem* ModifierEditorSubsystem = UActorModifierCoreEditorSubsystem::Get();
 
-	if (!ModifierSubsystem || !InAddToolMenu)
+	if (!ModifierSubsystem
+		|| !ModifierEditorSubsystem
+		|| !InAddToolMenu)
 	{
 		return;
 	}
@@ -36,15 +39,16 @@ void UE::ActorModifierCoreEditor::OnExtendAddModifierMenu(UToolMenu* InAddToolMe
 	for (const FName& Modifier : ShowModifiers)
 	{
 		const FName ModifierCategory = ModifierSubsystem->GetModifierCategory(Modifier);
+		TSharedPtr<const FActorModifierCoreEditorCategoryMetadata> Category = ModifierEditorSubsystem->FindModifierCategory(ModifierCategory);
+		const FText ModifierDisplayCategory = Category ? Category->DisplayName : FText::FromName(ModifierCategory);
+
 		FText ModifierDisplayName = FText::FromName(Modifier);
-		FText ModifierDisplayCategory = FText::FromName(ModifierCategory);
 		FText ModifierDescription = FText::GetEmpty();
 		FSlateIcon ModifierIcon = FSlateIconFinder::FindIconForClass(UActorModifierCoreBase::StaticClass());
 
-		ModifierSubsystem->ProcessModifierMetadata(Modifier, [&ModifierDisplayName, &ModifierDisplayCategory, &ModifierIcon, &ModifierDescription](const FActorModifierCoreMetadata& InMetadata)->bool
+		ModifierSubsystem->ProcessModifierMetadata(Modifier, [&ModifierDisplayName, &ModifierIcon, &ModifierDescription](const FActorModifierCoreMetadata& InMetadata)->bool
 		{
 			ModifierDisplayName = InMetadata.GetDisplayName();
-			ModifierDisplayCategory = FText::FromName(InMetadata.GetCategory());
 			ModifierIcon = InMetadata.GetIcon();
 			ModifierDescription = InMetadata.GetDescription();
 			return true;
@@ -86,8 +90,11 @@ void UE::ActorModifierCoreEditor::OnExtendAddModifierMenu(UToolMenu* InAddToolMe
 void UE::ActorModifierCoreEditor::OnExtendInsertModifierMenu(UToolMenu* InInsertToolMenu, const FActorModifierCoreEditorMenuData& InData)
 {
 	const UActorModifierCoreSubsystem* ModifierSubsystem = UActorModifierCoreSubsystem::Get();
+	const UActorModifierCoreEditorSubsystem* ModifierEditorSubsystem = UActorModifierCoreEditorSubsystem::Get();
 
-	if (!ModifierSubsystem || !InInsertToolMenu)
+	if (!ModifierSubsystem
+		|| !ModifierEditorSubsystem
+		|| !InInsertToolMenu)
 	{
 		return;
 	}
@@ -115,15 +122,16 @@ void UE::ActorModifierCoreEditor::OnExtendInsertModifierMenu(UToolMenu* InInsert
 	for (const FName& Modifier : ShowModifiers)
 	{
 		const FName ModifierCategory = ModifierSubsystem->GetModifierCategory(Modifier);
+		TSharedPtr<const FActorModifierCoreEditorCategoryMetadata> Category = ModifierEditorSubsystem->FindModifierCategory(ModifierCategory);
+		const FText ModifierDisplayCategory = Category ? Category->DisplayName : FText::FromName(ModifierCategory);
+
 		FText ModifierDisplayName = FText::FromName(Modifier);
-		FText ModifierDisplayCategory = FText::FromName(ModifierCategory);
 		FText ModifierDescription = FText::GetEmpty();
 		FSlateIcon ModifierIcon = FSlateIconFinder::FindIconForClass(UActorModifierCoreBase::StaticClass());
 
-		ModifierSubsystem->ProcessModifierMetadata(Modifier, [&ModifierDisplayName, &ModifierDisplayCategory, &ModifierIcon, &ModifierDescription](const FActorModifierCoreMetadata& InMetadata)->bool
+		ModifierSubsystem->ProcessModifierMetadata(Modifier, [&ModifierDisplayName, &ModifierIcon, &ModifierDescription](const FActorModifierCoreMetadata& InMetadata)->bool
 		{
 			ModifierDisplayName = InMetadata.GetDisplayName();
-			ModifierDisplayCategory = FText::FromName(InMetadata.GetCategory());
 			ModifierIcon = InMetadata.GetIcon();
 			ModifierDescription = InMetadata.GetDescription();
 			return true;
@@ -166,8 +174,11 @@ void UE::ActorModifierCoreEditor::OnExtendInsertModifierMenu(UToolMenu* InInsert
 void UE::ActorModifierCoreEditor::OnExtendRemoveModifierMenu(UToolMenu* InRemoveToolMenu, const FActorModifierCoreEditorMenuData& InData)
 {
 	const UActorModifierCoreSubsystem* ModifierSubsystem = UActorModifierCoreSubsystem::Get();
+	const UActorModifierCoreEditorSubsystem* ModifierEditorSubsystem = UActorModifierCoreEditorSubsystem::Get();
 
-	if (!ModifierSubsystem || !InRemoveToolMenu)
+	if (!ModifierSubsystem
+		|| !ModifierEditorSubsystem
+		|| !InRemoveToolMenu)
 	{
 		return;
 	}
@@ -252,15 +263,12 @@ void UE::ActorModifierCoreEditor::OnExtendRemoveModifierMenu(UToolMenu* InRemove
 		}
 
 		const FName ModifierName = StackModifier->GetModifierName();
-		const FName ModifierCategory = ModifierSubsystem->GetModifierCategory(ModifierName);
 		FText ModifierDisplayName = FText::FromName(ModifierName);
-		FText ModifierDisplayCategory = FText::FromName(ModifierCategory);
 		FSlateIcon ModifierIcon = FSlateIconFinder::FindIconForClass(UActorModifierCoreBase::StaticClass());
 
-		ModifierSubsystem->ProcessModifierMetadata(ModifierName, [&ModifierDisplayName, &ModifierDisplayCategory, &ModifierIcon](const FActorModifierCoreMetadata& InMetadata)->bool
+		ModifierSubsystem->ProcessModifierMetadata(ModifierName, [&ModifierDisplayName, &ModifierIcon](const FActorModifierCoreMetadata& InMetadata)->bool
 		{
 			ModifierDisplayName = InMetadata.GetDisplayName();
-			ModifierDisplayCategory = FText::FromName(InMetadata.GetCategory());
 			ModifierIcon = InMetadata.GetIcon();
 			return true;
 		});
@@ -287,8 +295,11 @@ void UE::ActorModifierCoreEditor::OnExtendRemoveModifierMenu(UToolMenu* InRemove
 void UE::ActorModifierCoreEditor::OnExtendMoveModifierMenu(UToolMenu* InMoveToolMenu, const FActorModifierCoreEditorMenuData& InData)
 {
 	const UActorModifierCoreSubsystem* ModifierSubsystem = UActorModifierCoreSubsystem::Get();
+	const UActorModifierCoreEditorSubsystem* ModifierEditorSubsystem = UActorModifierCoreEditorSubsystem::Get();
 
-	if (!ModifierSubsystem || !InMoveToolMenu)
+	if (!ModifierSubsystem
+		|| !ModifierEditorSubsystem
+		|| !InMoveToolMenu)
 	{
 		return;
 	}
@@ -347,15 +358,12 @@ void UE::ActorModifierCoreEditor::OnExtendMoveModifierMenu(UToolMenu* InMoveTool
 		}
 
 		const FName ModifierName = PositionModifier->GetModifierName();
-		const FName ModifierCategory = ModifierSubsystem->GetModifierCategory(ModifierName);
 		FText ModifierDisplayName = FText::FromName(ModifierName);
-		FText ModifierDisplayCategory = FText::FromName(ModifierCategory);
 		FSlateIcon ModifierIcon = FSlateIconFinder::FindIconForClass(UActorModifierCoreBase::StaticClass());
 
-		ModifierSubsystem->ProcessModifierMetadata(ModifierName, [&ModifierDisplayName, &ModifierDisplayCategory, &ModifierIcon](const FActorModifierCoreMetadata& InMetadata)->bool
+		ModifierSubsystem->ProcessModifierMetadata(ModifierName, [&ModifierDisplayName, &ModifierIcon](const FActorModifierCoreMetadata& InMetadata)->bool
 		{
 			ModifierDisplayName = InMetadata.GetDisplayName();
-			ModifierDisplayCategory = FText::FromName(InMetadata.GetCategory());
 			ModifierIcon = InMetadata.GetIcon();
 			return true;
 		});
diff --git a/Engine/Plugins/VirtualProduction/ActorModifierCore/Source/ActorModifierCoreEditor/Private/Modifiers/Customizations/ActorModifierCoreEditorStackCustomization.cpp b/Engine/Plugins/VirtualProduction/ActorModifierCore/Source/ActorModifierCoreEditor/Private/Modifiers/Customizations/ActorModifierCoreEditorStackCustomization.cpp
--- a/Engine/Plugins/VirtualProduction/ActorModifierCore/Source/ActorModifierCoreEditor/Private/Modifiers/Customizations/ActorModifierCoreEditorStackCustomization.cpp
+++ b/Engine/Plugins/VirtualProduction/ActorModifierCore/Source/ActorModifierCoreEditor/Private/Modifiers/Customizations/ActorModifierCoreEditorStackCustomization.cpp
@@ -251,20 +251,22 @@ void UActorModifierCoreEditorStackCustomization::CustomizeItemHeader(const FOper
 		/** Show last execution error messages if failed execution */
 		TAttribute<EOperatorStackEditorMessageType> MessageType = EOperatorStackEditorMessageType::None;
 		TAttribute<FText> MessageText = FText::GetEmpty();
-		FString HeaderLabel = bIsStack ? TEXT("Modifiers") : Modifier->GetModifierName().ToString();
+		FText HeaderLabel = bIsStack
+			? LOCTEXT("ModifierStackHeaderLabel", "Modifiers")
+			: FText::FromName(Modifier->GetModifierName());
 
 		ModifierSubsystem->ProcessModifierMetadata(
 			Modifier->GetModifierName(),
 			[&HeaderLabel](const FActorModifierCoreMetadata& InMetadata)->bool
 			{
-				HeaderLabel = InMetadata.GetDisplayName().ToString();
+				HeaderLabel = InMetadata.GetDisplayName();
 				return true;
 			}
 		);
 
 		if (InItem->GetValueCount() > 1)
 		{
-			HeaderLabel += TEXT(" (") + FString::FromInt(InItem->GetValueCount()) + TEXT(")");
+			HeaderLabel = FText::Format(LOCTEXT("ModifiersHeaderLabel", "{0} ({1})"), HeaderLabel, FText::AsNumber(InItem->GetValueCount()));
 
 			if (bIsStack)
 			{
@@ -320,7 +322,7 @@ void UActorModifierCoreEditorStackCustomization::CustomizeItemHeader(const FOper
 			.SetSearchKeywords(SearchKeywords)
 			.SetExpandable(!bIsStack)
 			.SetIcon(ModifierIcon.GetIcon())
-			.SetLabel(FText::FromString(HeaderLabel))
+			.SetLabel(HeaderLabel)
 			.SetTooltip(ModifierTooltip)
 			.SetBorderColor(ModifierColor)
 			.SetProperty(ModifierEnableProperty)
diff --git a/Engine/Plugins/VirtualProduction/ActorModifierCore/Source/ActorModifierCoreEditor/Private/Subsystems/ActorModifierCoreEditorSubsystem.cpp b/Engine/Plugins/VirtualProduction/ActorModifierCore/Source/ActorModifierCoreEditor/Private/Subsystems/ActorModifierCoreEditorSubsystem.cpp
--- a/Engine/Plugins/VirtualProduction/ActorModifierCore/Source/ActorModifierCoreEditor/Private/Subsystems/ActorModifierCoreEditorSubsystem.cpp
+++ b/Engine/Plugins/VirtualProduction/ActorModifierCore/Source/ActorModifierCoreEditor/Private/Subsystems/ActorModifierCoreEditorSubsystem.cpp
@@ -40,11 +40,57 @@ void UActorModifierCoreEditorSubsystem::Initialize(FSubsystemCollectionBase& Col
 
 	// Initialize here if not called anywhere to setup ClassIcons
 	FActorModifierCoreEditorStyle::Get();
-}
 
-void UActorModifierCoreEditorSubsystem::Deinitialize()
-{
-	Super::Deinitialize();
+	// Register modifier categories
+	{
+		const FActorModifierCoreEditorCategoryMetadata DefaultCategory
+		{
+			.Name = TEXT("Default"),
+			.DisplayName = LOCTEXT("DefaultCategoryDisplayName", "Default"),
+		};
+
+		RegisterModifierCategory(DefaultCategory);
+
+		const FActorModifierCoreEditorCategoryMetadata LayoutCategory
+		{
+			.Name = TEXT("Layout"),
+			.DisplayName = LOCTEXT("LayoutCategoryDisplayName", "Layout"),
+		};
+
+		RegisterModifierCategory(LayoutCategory);
+
+		const FActorModifierCoreEditorCategoryMetadata CameraCategory
+		{
+			.Name = TEXT("Camera"),
+			.DisplayName = LOCTEXT("CameraCategoryDisplayName", "Camera"),
+		};
+
+		RegisterModifierCategory(CameraCategory);
+
+		const FActorModifierCoreEditorCategoryMetadata RenderingCategory
+		{
+			.Name = TEXT("Rendering"),
+			.DisplayName = LOCTEXT("RenderingCategoryDisplayName", "Rendering"),
+		};
+
+		RegisterModifierCategory(RenderingCategory);
+
+		const FActorModifierCoreEditorCategoryMetadata GeometryCategory
+		{
+			.Name = TEXT("Geometry"),
+			.DisplayName = LOCTEXT("GeometryCategoryDisplayName", "Geometry"),
+		};
+
+		RegisterModifierCategory(GeometryCategory);
+
+		const FActorModifierCoreEditorCategoryMetadata ConversionCategory
+		{
+			.Name = TEXT("Conversion"),
+			.DisplayName = LOCTEXT("ConversionCategoryDisplayName", "Conversion"),
+		};
+
+		RegisterModifierCategory(ConversionCategory);
+	}
 }
 
 bool UActorModifierCoreEditorSubsystem::FillModifierMenu(UToolMenu* InMenu, const FActorModifierCoreEditorMenuContext& InContext, const FActorModifierCoreEditorMenuOptions& InMenuOptions) const
@@ -227,4 +273,31 @@ TSharedPtr<SActorModifierCoreEditorProfiler> UActorModifierCoreEditorSubsystem::
 	return SNew(SActorModifierCoreEditorProfiler, InProfiler);
 }
 
+bool UActorModifierCoreEditorSubsystem::RegisterModifierCategory(const FActorModifierCoreEditorCategoryMetadata& InCategoryMetadata)
+{
+	if (InCategoryMetadata.Name.IsNone())
+	{
+		return false;
+	}
+
+	if (ModifierCategories.Contains(InCategoryMetadata.Name))
+	{
+		return false;
+	}
+
+	ModifierCategories.Emplace(InCategoryMetadata.Name, MakeShared<const FActorModifierCoreEditorCategoryMetadata>(InCategoryMetadata));
+
+	return true;
+}
+
+TSharedPtr<const FActorModifierCoreEditorCategoryMetadata> UActorModifierCoreEditorSubsystem::FindModifierCategory(FName InCategoryIdentifier) const
+{
+	if (const TSharedRef<const FActorModifierCoreEditorCategoryMetadata>* Category = ModifierCategories.Find(InCategoryIdentifier))
+	{
+		return *Category;
+	}
+
+	return nullptr;
+}
+
 #undef LOCTEXT_NAMESPACE
diff --git a/Engine/Plugins/VirtualProduction/ActorModifierCore/Source/ActorModifierCoreEditor/Public/Subsystems/ActorModifierCoreEditorSubsystem.h b/Engine/Plugins/VirtualProduction/ActorModifierCore/Source/ActorModifierCoreEditor/Public/Subsystems/ActorModifierCoreEditorSubsystem.h
--- a/Engine/Plugins/VirtualProduction/ActorModifierCore/Source/ActorModifierCoreEditor/Public/Subsystems/ActorModifierCoreEditorSubsystem.h
+++ b/Engine/Plugins/VirtualProduction/ActorModifierCore/Source/ActorModifierCoreEditor/Public/Subsystems/ActorModifierCoreEditorSubsystem.h
@@ -17,6 +17,12 @@ class UActorModifierCoreBase;
 class UActorModifierCoreStack;
 class UActorModifierCoreSubsystem;
 
+struct FActorModifierCoreEditorCategoryMetadata
+{
+	FName Name;
+	FText DisplayName;
+};
+
 /** Singleton class that handles editor operations for modifiers */
 UCLASS()
 class UActorModifierCoreEditorSubsystem : public UEditorSubsystem
@@ -30,11 +36,6 @@ public:
 
 	ACTORMODIFIERCOREEDITOR_API static UActorModifierCoreEditorSubsystem* Get();
 
-	//~ Begin USubsystem
-	virtual void Initialize(FSubsystemCollectionBase& Collection) override;
-	virtual void Deinitialize() override;
-	//~ End USubsystem
-
 	/** Fills a menu based on context objects and menu options, will perform menu action transaction if wanted, returns true if the menu has been modified */
 	ACTORMODIFIERCOREEDITOR_API bool FillModifierMenu(UToolMenu* InMenu, const FActorModifierCoreEditorMenuContext& InContext, const FActorModifierCoreEditorMenuOptions& InMenuOptions) const;
 
@@ -61,13 +62,26 @@ public:
 	/** Creates a profiler widget for a modifier profiler, returns default widget if none was registered */
 	TSharedPtr<SActorModifierCoreEditorProfiler> CreateProfilerWidget(TSharedPtr<FActorModifierCoreProfiler> InProfiler);
 
+	/** Register a category for modifier */
+	ACTORMODIFIERCOREEDITOR_API bool RegisterModifierCategory(const FActorModifierCoreEditorCategoryMetadata& InCategoryMetadata);
+
+	/** Finds a category for modifier */
+	ACTORMODIFIERCOREEDITOR_API TSharedPtr<const FActorModifierCoreEditorCategoryMetadata> FindModifierCategory(FName InCategoryIdentifier) const;
+
 protected:
+	//~ Begin USubsystem
+	virtual void Initialize(FSubsystemCollectionBase& Collection) override;
+	//~ End USubsystem
+
 	/** Modifier profiler pinned stats based on profiler type */
 	TMap<FName, TSet<FName>> ModifierProfilerStats;
 
 	/** Modifier profiler widgets */
 	TMap<FName, TFunction<TSharedRef<SActorModifierCoreEditorProfiler>(TSharedPtr<FActorModifierCoreProfiler>)>> ModifierProfilerWidgets;
 
+	/** Registered modifier category for menus and widgets */
+	TMap<FName, TSharedRef<const FActorModifierCoreEditorCategoryMetadata>> ModifierCategories;
+
 	/** Runtime subsystem for modifier factories */
 	TWeakObjectPtr<UActorModifierCoreSubsystem> EngineSubsystem = nullptr;
 };
diff --git a/Engine/Plugins/VirtualProduction/Avalanche/Source/AvalancheMask/Private/Mask2D/AvaMask2DReadModifier.cpp b/Engine/Plugins/VirtualProduction/Avalanche/Source/AvalancheMask/Private/Mask2D/AvaMask2DReadModifier.cpp
--- a/Engine/Plugins/VirtualProduction/Avalanche/Source/AvalancheMask/Private/Mask2D/AvaMask2DReadModifier.cpp
+++ b/Engine/Plugins/VirtualProduction/Avalanche/Source/AvalancheMask/Private/Mask2D/AvaMask2DReadModifier.cpp
@@ -161,7 +161,7 @@ void UAvaMask2DReadModifier::OnModifierCDOSetup(FActorModifierCoreMetadata& InMe
 	InMetadata.DisallowAfter(TEXT("MaskWrite"));
 	InMetadata.DisallowBefore(TEXT("MaskWrite"));
 #if WITH_EDITOR
-	InMetadata.SetDisplayName(FText::FromString(TEXT("Masked Layer (Output)")));
+	InMetadata.SetDisplayName(LOCTEXT("ModifierDisplayName", "Masked Layer (Output)"));
 	InMetadata.SetDescription(LOCTEXT("ModifierDescription", "Reads from a canvas texture and uses it on materials"));
 #endif
 }
diff --git a/Engine/Plugins/VirtualProduction/Avalanche/Source/AvalancheMask/Private/Mask2D/AvaMask2DWriteModifier.cpp b/Engine/Plugins/VirtualProduction/Avalanche/Source/AvalancheMask/Private/Mask2D/AvaMask2DWriteModifier.cpp
--- a/Engine/Plugins/VirtualProduction/Avalanche/Source/AvalancheMask/Private/Mask2D/AvaMask2DWriteModifier.cpp
+++ b/Engine/Plugins/VirtualProduction/Avalanche/Source/AvalancheMask/Private/Mask2D/AvaMask2DWriteModifier.cpp
@@ -112,7 +112,7 @@ void UAvaMask2DWriteModifier::OnModifierCDOSetup(FActorModifierCoreMetadata& InM
 	InMetadata.DisallowAfter(TEXT("MaskRead"));
 	InMetadata.DisallowBefore(TEXT("MaskRead"));
 #if WITH_EDITOR
-	InMetadata.SetDisplayName(FText::FromString(TEXT("Mask Layer (Input)")));
+	InMetadata.SetDisplayName(LOCTEXT("ModifierDisplayName", "Mask Layer (Input)"));
 	InMetadata.SetDescription(LOCTEXT("ModifierDescription", "Writes to a canvas texture used by the masked layer in materials"));
 #endif
 }
diff --git a/Engine/Plugins/VirtualProduction/Avalanche/Source/AvalancheModifiers/Private/Modifiers/AvaAutoSizeModifier.cpp b/Engine/Plugins/VirtualProduction/Avalanche/Source/AvalancheModifiers/Private/Modifiers/AvaAutoSizeModifier.cpp
--- a/Engine/Plugins/VirtualProduction/Avalanche/Source/AvalancheModifiers/Private/Modifiers/AvaAutoSizeModifier.cpp
+++ b/Engine/Plugins/VirtualProduction/Avalanche/Source/AvalancheModifiers/Private/Modifiers/AvaAutoSizeModifier.cpp
@@ -41,6 +41,7 @@ void UAvaAutoSizeModifier::OnModifierCDOSetup(FActorModifierCoreMetadata& InMeta
 	InMetadata.SetName(TEXT("AutoSize"));
 	InMetadata.SetCategory(TEXT("Geometry"));
 #if WITH_EDITOR
+	InMetadata.SetDisplayName(LOCTEXT("ModifierDisplayName", "Auto Size"));
 	InMetadata.SetDescription(LOCTEXT("ModifierDescription", "The modified actor will be resized to act as a background for a reference actor"));
 #endif
 
diff --git a/Engine/Plugins/VirtualProduction/Avalanche/Source/AvalancheModifiers/Private/Modifiers/AvaBendModifier.cpp b/Engine/Plugins/VirtualProduction/Avalanche/Source/AvalancheModifiers/Private/Modifiers/AvaBendModifier.cpp
--- a/Engine/Plugins/VirtualProduction/Avalanche/Source/AvalancheModifiers/Private/Modifiers/AvaBendModifier.cpp
+++ b/Engine/Plugins/VirtualProduction/Avalanche/Source/AvalancheModifiers/Private/Modifiers/AvaBendModifier.cpp
@@ -17,6 +17,7 @@ void UAvaBendModifier::OnModifierCDOSetup(FActorModifierCoreMetadata& InMetadata
 	InMetadata.SetName(TEXT("Bend"));
 	InMetadata.SetCategory(TEXT("Geometry"));
 #if WITH_EDITOR
+	InMetadata.SetDisplayName(LOCTEXT("ModifierDisplayName", "Bend"));
 	InMetadata.SetDescription(LOCTEXT("ModifierDescription", "Bend the current geometry shape with a transition between two sides"));
 #endif
 	InMetadata.AddDependency(TEXT("Subdivide"));
diff --git a/Engine/Plugins/VirtualProduction/Avalanche/Source/AvalancheModifiers/Private/Modifiers/AvaBevelModifier.cpp b/Engine/Plugins/VirtualProduction/Avalanche/Source/AvalancheModifiers/Private/Modifiers/AvaBevelModifier.cpp
--- a/Engine/Plugins/VirtualProduction/Avalanche/Source/AvalancheModifiers/Private/Modifiers/AvaBevelModifier.cpp
+++ b/Engine/Plugins/VirtualProduction/Avalanche/Source/AvalancheModifiers/Private/Modifiers/AvaBevelModifier.cpp
@@ -19,6 +19,7 @@ void UAvaBevelModifier::OnModifierCDOSetup(FActorModifierCoreMetadata& InMetadat
 	InMetadata.SetName(TEXT("Bevel"));
 	InMetadata.SetCategory(TEXT("Geometry"));
 #if WITH_EDITOR
+	InMetadata.SetDisplayName(LOCTEXT("ModifierDisplayName", "Bevel"));
 	InMetadata.SetDescription(LOCTEXT("ModifierDescription", "Create chamfered or rounded corners on geometry that smooth edges and corners"));
 #endif
 }
diff --git a/Engine/Plugins/VirtualProduction/Avalanche/Source/AvalancheModifiers/Private/Modifiers/AvaBooleanModifier.cpp b/Engine/Plugins/VirtualProduction/Avalanche/Source/AvalancheModifiers/Private/Modifiers/AvaBooleanModifier.cpp
--- a/Engine/Plugins/VirtualProduction/Avalanche/Source/AvalancheModifiers/Private/Modifiers/AvaBooleanModifier.cpp
+++ b/Engine/Plugins/VirtualProduction/Avalanche/Source/AvalancheModifiers/Private/Modifiers/AvaBooleanModifier.cpp
@@ -41,6 +41,7 @@ void UAvaBooleanModifier::OnModifierCDOSetup(FActorModifierCoreMetadata& InMetad
 	InMetadata.SetName(TEXT("Boolean"));
 	InMetadata.SetCategory(TEXT("Geometry"));
 #if WITH_EDITOR
+	InMetadata.SetDisplayName(LOCTEXT("ModifierDisplayName", "Boolean"));
 	InMetadata.SetDescription(LOCTEXT("ModifierDescription", "Subtracts or intersects a geometry shape with another one when they collide"));
 #endif
 }
diff --git a/Engine/Plugins/VirtualProduction/Avalanche/Source/AvalancheModifiers/Private/Modifiers/AvaDynamicMeshConverterModifier.cpp b/Engine/Plugins/VirtualProduction/Avalanche/Source/AvalancheModifiers/Private/Modifiers/AvaDynamicMeshConverterModifier.cpp
--- a/Engine/Plugins/VirtualProduction/Avalanche/Source/AvalancheModifiers/Private/Modifiers/AvaDynamicMeshConverterModifier.cpp
+++ b/Engine/Plugins/VirtualProduction/Avalanche/Source/AvalancheModifiers/Private/Modifiers/AvaDynamicMeshConverterModifier.cpp
@@ -49,6 +49,7 @@ void UAvaDynamicMeshConverterModifier::OnModifierCDOSetup(FActorModifierCoreMeta
 	InMetadata.SetCategory(TEXT("Conversion"));
 	InMetadata.AllowTick(true);
 #if WITH_EDITOR
+	InMetadata.SetDisplayName(LOCTEXT("ModifierDisplayName", "Dynamic Mesh Converter"));
 	InMetadata.SetDescription(LOCTEXT("ModifierDescription", "Converts various actor mesh types into a single dynamic mesh, this is an heavy operation"));
 #endif
 	InMetadata.SetCompatibilityRule([](const AActor* InActor)->bool
diff --git a/Engine/Plugins/VirtualProduction/Avalanche/Source/AvalancheModifiers/Private/Modifiers/AvaExtrudeModifier.cpp b/Engine/Plugins/VirtualProduction/Avalanche/Source/AvalancheModifiers/Private/Modifiers/AvaExtrudeModifier.cpp
--- a/Engine/Plugins/VirtualProduction/Avalanche/Source/AvalancheModifiers/Private/Modifiers/AvaExtrudeModifier.cpp
+++ b/Engine/Plugins/VirtualProduction/Avalanche/Source/AvalancheModifiers/Private/Modifiers/AvaExtrudeModifier.cpp
@@ -41,6 +41,7 @@ void UAvaExtrudeModifier::OnModifierCDOSetup(FActorModifierCoreMetadata& InMetad
 	InMetadata.SetName(TEXT("Extrude"));
 	InMetadata.SetCategory(TEXT("Geometry"));
 #if WITH_EDITOR
+	InMetadata.SetDisplayName(LOCTEXT("ModifierDisplayName", "Extrude"));
 	InMetadata.SetDescription(LOCTEXT("ModifierDescription", "Duplicates and moves points and triangles following their normal direction"));
 #endif
 
diff --git a/Engine/Plugins/VirtualProduction/Avalanche/Source/AvalancheModifiers/Private/Modifiers/AvaGlobalOpacityModifier.cpp b/Engine/Plugins/VirtualProduction/Avalanche/Source/AvalancheModifiers/Private/Modifiers/AvaGlobalOpacityModifier.cpp
--- a/Engine/Plugins/VirtualProduction/Avalanche/Source/AvalancheModifiers/Private/Modifiers/AvaGlobalOpacityModifier.cpp
+++ b/Engine/Plugins/VirtualProduction/Avalanche/Source/AvalancheModifiers/Private/Modifiers/AvaGlobalOpacityModifier.cpp
@@ -136,6 +136,7 @@ void UAvaGlobalOpacityModifier::OnModifierCDOSetup(FActorModifierCoreMetadata& I
 	InMetadata.SetName(TEXT("GlobalOpacity"));
 	InMetadata.SetCategory(TEXT("Rendering"));
 #if WITH_EDITOR
+	InMetadata.SetDisplayName(LOCTEXT("ModifierDisplayName", "Global Opacity"));
 	InMetadata.SetDescription(LOCTEXT("ModifierDescription", "Sets global opacity parameters on an actor with Material Designer Instances generated with the Material Designer"));
 #endif
 }
diff --git a/Engine/Plugins/VirtualProduction/Avalanche/Source/AvalancheModifiers/Private/Modifiers/AvaHideEmptyModifier.cpp b/Engine/Plugins/VirtualProduction/Avalanche/Source/AvalancheModifiers/Private/Modifiers/AvaHideEmptyModifier.cpp
--- a/Engine/Plugins/VirtualProduction/Avalanche/Source/AvalancheModifiers/Private/Modifiers/AvaHideEmptyModifier.cpp
+++ b/Engine/Plugins/VirtualProduction/Avalanche/Source/AvalancheModifiers/Private/Modifiers/AvaHideEmptyModifier.cpp
@@ -14,6 +14,7 @@ void UAvaHideEmptyModifier::OnModifierCDOSetup(FActorModifierCoreMetadata& InMet
 	InMetadata.SetName(TEXT("HideEmpty"));
 	InMetadata.SetCategory(TEXT("Rendering"));
 #if WITH_EDITOR
+	InMetadata.SetDisplayName(LOCTEXT("ModifierDisplayName", "Hide Empty"));
 	InMetadata.SetDescription(LOCTEXT("ModifierDescription", "Hides a container when the text content is empty"));
 #endif
 
diff --git a/Engine/Plugins/VirtualProduction/Avalanche/Source/AvalancheModifiers/Private/Modifiers/AvaMaterialParameterModifier.cpp b/Engine/Plugins/VirtualProduction/Avalanche/Source/AvalancheModifiers/Private/Modifiers/AvaMaterialParameterModifier.cpp
--- a/Engine/Plugins/VirtualProduction/Avalanche/Source/AvalancheModifiers/Private/Modifiers/AvaMaterialParameterModifier.cpp
+++ b/Engine/Plugins/VirtualProduction/Avalanche/Source/AvalancheModifiers/Private/Modifiers/AvaMaterialParameterModifier.cpp
@@ -121,6 +121,7 @@ void UAvaMaterialParameterModifier::OnModifierCDOSetup(FActorModifierCoreMetadat
 	InMetadata.SetName(TEXT("MaterialParameter"));
 	InMetadata.SetCategory(TEXT("Rendering"));
 #if WITH_EDITOR
+	InMetadata.SetDisplayName(LOCTEXT("ModifierDisplayName", "Material Parameter"));
 	InMetadata.SetDescription(LOCTEXT("ModifierDescription", "Sets specified dynamic materials parameters on an actor and its children"));
 #endif
 
diff --git a/Engine/Plugins/VirtualProduction/Avalanche/Source/AvalancheModifiers/Private/Modifiers/AvaMirrorModifier.cpp b/Engine/Plugins/VirtualProduction/Avalanche/Source/AvalancheModifiers/Private/Modifiers/AvaMirrorModifier.cpp
--- a/Engine/Plugins/VirtualProduction/Avalanche/Source/AvalancheModifiers/Private/Modifiers/AvaMirrorModifier.cpp
+++ b/Engine/Plugins/VirtualProduction/Avalanche/Source/AvalancheModifiers/Private/Modifiers/AvaMirrorModifier.cpp
@@ -18,6 +18,7 @@ void UAvaMirrorModifier::OnModifierCDOSetup(FActorModifierCoreMetadata& InMetada
 	InMetadata.SetName(TEXT("Mirror"));
 	InMetadata.SetCategory(TEXT("Geometry"));
 #if WITH_EDITOR
+	InMetadata.SetDisplayName(LOCTEXT("ModifierDisplayName", "Mirror"));
 	InMetadata.SetDescription(LOCTEXT("ModifierDescription", "Duplicates and mirrors a geometry shape along a plane"));
 #endif
 }
diff --git a/Engine/Plugins/VirtualProduction/Avalanche/Source/AvalancheModifiers/Private/Modifiers/AvaNormalModifier.cpp b/Engine/Plugins/VirtualProduction/Avalanche/Source/AvalancheModifiers/Private/Modifiers/AvaNormalModifier.cpp
--- a/Engine/Plugins/VirtualProduction/Avalanche/Source/AvalancheModifiers/Private/Modifiers/AvaNormalModifier.cpp
+++ b/Engine/Plugins/VirtualProduction/Avalanche/Source/AvalancheModifiers/Private/Modifiers/AvaNormalModifier.cpp
@@ -16,6 +16,7 @@ void UAvaNormalModifier::OnModifierCDOSetup(FActorModifierCoreMetadata& InMetada
 	InMetadata.SetName(TEXT("Normal"));
 	InMetadata.SetCategory(TEXT("Geometry"));
 #if WITH_EDITOR
+	InMetadata.SetDisplayName(LOCTEXT("ModifierDisplayName", "Normal"));
 	InMetadata.SetDescription(LOCTEXT("ModifierDescription", "Recompute the normals and split them based on different options"));
 #endif
 }
diff --git a/Engine/Plugins/VirtualProduction/Avalanche/Source/AvalancheModifiers/Private/Modifiers/AvaOutlineModifier.cpp b/Engine/Plugins/VirtualProduction/Avalanche/Source/AvalancheModifiers/Private/Modifiers/AvaOutlineModifier.cpp
--- a/Engine/Plugins/VirtualProduction/Avalanche/Source/AvalancheModifiers/Private/Modifiers/AvaOutlineModifier.cpp
+++ b/Engine/Plugins/VirtualProduction/Avalanche/Source/AvalancheModifiers/Private/Modifiers/AvaOutlineModifier.cpp
@@ -20,6 +20,7 @@ void UAvaOutlineModifier::OnModifierCDOSetup(FActorModifierCoreMetadata& InMetad
 	InMetadata.SetName(TEXT("Outline"));
 	InMetadata.SetCategory(TEXT("Geometry"));
 #if WITH_EDITOR
+	InMetadata.SetDisplayName(LOCTEXT("ModifierDisplayName", "Outline"));
 	InMetadata.SetDescription(LOCTEXT("ModifierDescription", "Creates a new region inside or outside of a geometry shape"));
 #endif
 
diff --git a/Engine/Plugins/VirtualProduction/Avalanche/Source/AvalancheModifiers/Private/Modifiers/AvaPatternModifier.cpp b/Engine/Plugins/VirtualProduction/Avalanche/Source/AvalancheModifiers/Private/Modifiers/AvaPatternModifier.cpp
--- a/Engine/Plugins/VirtualProduction/Avalanche/Source/AvalancheModifiers/Private/Modifiers/AvaPatternModifier.cpp
+++ b/Engine/Plugins/VirtualProduction/Avalanche/Source/AvalancheModifiers/Private/Modifiers/AvaPatternModifier.cpp
@@ -35,6 +35,7 @@ void UAvaPatternModifier::OnModifierCDOSetup(FActorModifierCoreMetadata& InMetad
 	InMetadata.SetName(TEXT("Pattern"));
 	InMetadata.SetCategory(TEXT("Geometry"));
 #if WITH_EDITOR
+	InMetadata.SetDisplayName(LOCTEXT("ModifierDisplayName", "Pattern"));
 	InMetadata.SetDescription(LOCTEXT("ModifierDescription", "Repeats a geometry multiple times following a specific layout pattern"));
 #endif
 }
diff --git a/Engine/Plugins/VirtualProduction/Avalanche/Source/AvalancheModifiers/Private/Modifiers/AvaPlaneCutModifier.cpp b/Engine/Plugins/VirtualProduction/Avalanche/Source/AvalancheModifiers/Private/Modifiers/AvaPlaneCutModifier.cpp
--- a/Engine/Plugins/VirtualProduction/Avalanche/Source/AvalancheModifiers/Private/Modifiers/AvaPlaneCutModifier.cpp
+++ b/Engine/Plugins/VirtualProduction/Avalanche/Source/AvalancheModifiers/Private/Modifiers/AvaPlaneCutModifier.cpp
@@ -153,6 +153,7 @@ void UAvaPlaneCutModifier::OnModifierCDOSetup(FActorModifierCoreMetadata& InMeta
 	InMetadata.SetName(TEXT("PlaneCut"));
 	InMetadata.SetCategory(TEXT("Geometry"));
 #if WITH_EDITOR
+	InMetadata.SetDisplayName(LOCTEXT("ModifierDisplayName", "Plane Cut"));
 	InMetadata.SetDescription(LOCTEXT("ModifierDescription", "Cuts a shape based on a 2D plane"));
 #endif
 }
diff --git a/Engine/Plugins/VirtualProduction/Avalanche/Source/AvalancheModifiers/Private/Modifiers/AvaSubdivideModifier.cpp b/Engine/Plugins/VirtualProduction/Avalanche/Source/AvalancheModifiers/Private/Modifiers/AvaSubdivideModifier.cpp
--- a/Engine/Plugins/VirtualProduction/Avalanche/Source/AvalancheModifiers/Private/Modifiers/AvaSubdivideModifier.cpp
+++ b/Engine/Plugins/VirtualProduction/Avalanche/Source/AvalancheModifiers/Private/Modifiers/AvaSubdivideModifier.cpp
@@ -18,6 +18,7 @@ void UAvaSubdivideModifier::OnModifierCDOSetup(FActorModifierCoreMetadata& InMet
 	InMetadata.SetName(TEXT("Subdivide"));
 	InMetadata.SetCategory(TEXT("Geometry"));
 #if WITH_EDITOR
+	InMetadata.SetDisplayName(LOCTEXT("ModifierDisplayName", "Subdivide"));
 	InMetadata.SetDescription(LOCTEXT("ModifierDescription", "Adds resolution to the geometry shape by dividing faces or edges into smaller units"));
 #endif
 }
diff --git a/Engine/Plugins/VirtualProduction/Avalanche/Source/AvalancheModifiers/Private/Modifiers/AvaTaperModifier.cpp b/Engine/Plugins/VirtualProduction/Avalanche/Source/AvalancheModifiers/Private/Modifiers/AvaTaperModifier.cpp
--- a/Engine/Plugins/VirtualProduction/Avalanche/Source/AvalancheModifiers/Private/Modifiers/AvaTaperModifier.cpp
+++ b/Engine/Plugins/VirtualProduction/Avalanche/Source/AvalancheModifiers/Private/Modifiers/AvaTaperModifier.cpp
@@ -18,6 +18,7 @@ void UAvaTaperModifier::OnModifierCDOSetup(FActorModifierCoreMetadata& InMetadat
 	InMetadata.SetName(TEXT("Taper"));
 	InMetadata.SetCategory(TEXT("Geometry"));
 #if WITH_EDITOR
+	InMetadata.SetDisplayName(LOCTEXT("ModifierDisplayName", "Taper"));
 	InMetadata.SetDescription(LOCTEXT("ModifierDescription", "Moves the vertices location of a geometry shape using a curve"));
 #endif
 
diff --git a/Engine/Plugins/VirtualProduction/Avalanche/Source/AvalancheModifiers/Private/Modifiers/AvaTranslucentPriorityModifier.cpp b/Engine/Plugins/VirtualProduction/Avalanche/Source/AvalancheModifiers/Private/Modifiers/AvaTranslucentPriorityModifier.cpp
--- a/Engine/Plugins/VirtualProduction/Avalanche/Source/AvalancheModifiers/Private/Modifiers/AvaTranslucentPriorityModifier.cpp
+++ b/Engine/Plugins/VirtualProduction/Avalanche/Source/AvalancheModifiers/Private/Modifiers/AvaTranslucentPriorityModifier.cpp
@@ -28,6 +28,7 @@ void UAvaTranslucentPriorityModifier::OnModifierCDOSetup(FActorModifierCoreMetad
 	InMetadata.SetName(TEXT("TranslucentSortPriority"));
 	InMetadata.SetCategory(TEXT("Rendering"));
 #if WITH_EDITOR
+	InMetadata.SetDisplayName(LOCTEXT("ModifierDisplayName", "Translucent Sort Priority"));
 	InMetadata.SetDescription(LOCTEXT("ModifierDescription", "Set the translucency sort priority of a primitive component based on different modes"));
 #endif
 
diff --git a/Engine/Plugins/VirtualProduction/Avalanche/Source/AvalancheModifiers/Private/Modifiers/AvaVisibilityModifier.cpp b/Engine/Plugins/VirtualProduction/Avalanche/Source/AvalancheModifiers/Private/Modifiers/AvaVisibilityModifier.cpp
--- a/Engine/Plugins/VirtualProduction/Avalanche/Source/AvalancheModifiers/Private/Modifiers/AvaVisibilityModifier.cpp
+++ b/Engine/Plugins/VirtualProduction/Avalanche/Source/AvalancheModifiers/Private/Modifiers/AvaVisibilityModifier.cpp
@@ -175,6 +175,7 @@ void UAvaVisibilityModifier::OnModifierCDOSetup(FActorModifierCoreMetadata& InMe
 	InMetadata.SetName(TEXT("Visibility"));
 	InMetadata.SetCategory(TEXT("Rendering"));
 #if WITH_EDITOR
+	InMetadata.SetDisplayName(LOCTEXT("ModifierDisplayName", "Visibility"));
 	InMetadata.SetDescription(LOCTEXT("ModifierDescription", "Controls the visibility of a range of child actors by index"));
 #endif
 }
