Motion Design: moved asset definitions not inheriting from the default to now inherit from default to adopt the asset status icon logic.

diff --git a/Engine/Plugins/Experimental/DataLink/Source/DataLinkEditor/Private/AssetDefinition_DataLinkGraph.h b/Engine/Plugins/Experimental/DataLink/Source/DataLinkEditor/Private/AssetDefinition_DataLinkGraph.h
--- a/Engine/Plugins/Experimental/DataLink/Source/DataLinkEditor/Private/AssetDefinition_DataLinkGraph.h
+++ b/Engine/Plugins/Experimental/DataLink/Source/DataLinkEditor/Private/AssetDefinition_DataLinkGraph.h
@@ -2,11 +2,11 @@
 
 #pragma once
 
-#include "AssetDefinition.h"
+#include "AssetDefinitionDefault.h"
 #include "AssetDefinition_DataLinkGraph.generated.h"
 
 UCLASS()
-class UAssetDefinition_DataLinkGraph : public UAssetDefinition
+class UAssetDefinition_DataLinkGraph : public UAssetDefinitionDefault
 {
 	GENERATED_BODY()
 
diff --git a/Engine/Plugins/Experimental/SceneState/Source/SceneStateBlueprintEditor/Private/AssetDefinition_SceneStateBlueprint.h b/Engine/Plugins/Experimental/SceneState/Source/SceneStateBlueprintEditor/Private/AssetDefinition_SceneStateBlueprint.h
--- a/Engine/Plugins/Experimental/SceneState/Source/SceneStateBlueprintEditor/Private/AssetDefinition_SceneStateBlueprint.h
+++ b/Engine/Plugins/Experimental/SceneState/Source/SceneStateBlueprintEditor/Private/AssetDefinition_SceneStateBlueprint.h
@@ -2,11 +2,11 @@
 
 #pragma once
 
-#include "AssetDefinition.h"
+#include "AssetDefinitionDefault.h"
 #include "AssetDefinition_SceneStateBlueprint.generated.h"
 
 UCLASS()
-class UAssetDefinition_SceneStateBlueprint : public UAssetDefinition
+class UAssetDefinition_SceneStateBlueprint : public UAssetDefinitionDefault
 {
 	GENERATED_BODY()
 
diff --git a/Engine/Plugins/Experimental/SceneState/Source/SceneStateBlueprintEditor/Private/AssetDefinition_SceneStateTaskBlueprint.h b/Engine/Plugins/Experimental/SceneState/Source/SceneStateBlueprintEditor/Private/AssetDefinition_SceneStateTaskBlueprint.h
--- a/Engine/Plugins/Experimental/SceneState/Source/SceneStateBlueprintEditor/Private/AssetDefinition_SceneStateTaskBlueprint.h
+++ b/Engine/Plugins/Experimental/SceneState/Source/SceneStateBlueprintEditor/Private/AssetDefinition_SceneStateTaskBlueprint.h
@@ -2,11 +2,11 @@
 
 #pragma once
 
-#include "AssetDefinition.h"
+#include "AssetDefinitionDefault.h"
 #include "AssetDefinition_SceneStateTaskBlueprint.generated.h"
 
 UCLASS()
-class UAssetDefinition_SceneStateTaskBlueprint : public UAssetDefinition
+class UAssetDefinition_SceneStateTaskBlueprint : public UAssetDefinitionDefault
 {
 	GENERATED_BODY()
 
diff --git a/Engine/Plugins/VirtualProduction/ActorModifierCore/Source/ActorModifierCoreEditor/Private/Modifiers/Blueprints/AssetDefinition_ActorModifierCoreBlueprint.h b/Engine/Plugins/VirtualProduction/ActorModifierCore/Source/ActorModifierCoreEditor/Private/Modifiers/Blueprints/AssetDefinition_ActorModifierCoreBlueprint.h
--- a/Engine/Plugins/VirtualProduction/ActorModifierCore/Source/ActorModifierCoreEditor/Private/Modifiers/Blueprints/AssetDefinition_ActorModifierCoreBlueprint.h
+++ b/Engine/Plugins/VirtualProduction/ActorModifierCore/Source/ActorModifierCoreEditor/Private/Modifiers/Blueprints/AssetDefinition_ActorModifierCoreBlueprint.h
@@ -2,12 +2,12 @@
 
 #pragma once
 
-#include "AssetDefinition.h"
+#include "AssetDefinitionDefault.h"
 #include "AssetDefinition_ActorModifierCoreBlueprint.generated.h"
 
 /** Asset definition for UActorModifierCoreBlueprint objects available in editor */
 UCLASS()
-class UAssetDefinition_ActorModifierCoreBlueprint : public UAssetDefinition
+class UAssetDefinition_ActorModifierCoreBlueprint : public UAssetDefinitionDefault
 {
 	GENERATED_BODY()
 
