MotionDesign : ActorModifier
- Added a way to mark modifier dirty in blueprint

diff --git a/Engine/Plugins/VirtualProduction/ActorModifierCore/Source/ActorModifierCore/Private/Modifiers/Utilities/ActorModifierCoreLibrary.cpp b/Engine/Plugins/VirtualProduction/ActorModifierCore/Source/ActorModifierCore/Private/Modifiers/Utilities/ActorModifierCoreLibrary.cpp
--- a/Engine/Plugins/VirtualProduction/ActorModifierCore/Source/ActorModifierCore/Private/Modifiers/Utilities/ActorModifierCoreLibrary.cpp
+++ b/Engine/Plugins/VirtualProduction/ActorModifierCore/Source/ActorModifierCore/Private/Modifiers/Utilities/ActorModifierCoreLibrary.cpp
@@ -8,6 +8,14 @@
 
 DEFINE_LOG_CATEGORY_STATIC(LogActorModifierCoreLibrary, Log, All);
 
+void UActorModifierCoreLibrary::MarkModifierDirty(UActorModifierCoreBase* InModifier)
+{
+	if (InModifier)
+	{
+		InModifier->MarkModifierDirty();
+	}
+}
+
 bool UActorModifierCoreLibrary::FindModifierStack(AActor* InActor, UActorModifierCoreStack*& OutModifierStack, bool bInCreateIfNone)
 {
 	OutModifierStack = nullptr;
diff --git a/Engine/Plugins/VirtualProduction/ActorModifierCore/Source/ActorModifierCore/Public/Modifiers/Utilities/ActorModifierCoreLibrary.h b/Engine/Plugins/VirtualProduction/ActorModifierCore/Source/ActorModifierCore/Public/Modifiers/Utilities/ActorModifierCoreLibrary.h
--- a/Engine/Plugins/VirtualProduction/ActorModifierCore/Source/ActorModifierCore/Public/Modifiers/Utilities/ActorModifierCoreLibrary.h
+++ b/Engine/Plugins/VirtualProduction/ActorModifierCore/Source/ActorModifierCore/Public/Modifiers/Utilities/ActorModifierCoreLibrary.h
@@ -16,6 +16,13 @@ class UActorModifierCoreLibrary : public UBlueprintFunctionLibrary
 	GENERATED_BODY()
 
 public:
+	/**
+	 * Marks the modifier dirty
+	 * @param InModifier The modifier to dirty
+	 */
+	UFUNCTION(BlueprintCallable, Category="Motion Design|Modifiers|Utility")
+	static ACTORMODIFIERCORE_API void MarkModifierDirty(UActorModifierCoreBase* InModifier);
+
 	/**
 	 * Retrieves the modifier stack or create one if none is found
 	 * @param InActor The actor to get the modifier stack from
