Motion Design: fix issue with mask modifiers setting a wrong state to actors in MRQ PIE.
This issue is caused by the following:
- PIE duplicates the world to use. The mask modifier 'actor data' is 'DuplicateTransient' which means it will be cleared out and no original state is kept. The solve for this is to change the specifier to 'NonPIEDuplicateTransient', as we only want to clear it out when the modifier gets duplicated by the user and not when creating the PIE world.
- MRQ does not do actor ticks for its PIE world so the delegate OnWorldPostActorTick never gets broadcast. Because of this InitializeModifier with 'Load' reason never gets called. This init call is critical for the mask modifier to work as intended. The solve for this was to use a delegate that will always get called in UWorld::Tick.

diff --git a/Engine/Plugins/VirtualProduction/ActorModifierCore/Source/ActorModifierCore/Private/Modifiers/ActorModifierCoreBase.cpp b/Engine/Plugins/VirtualProduction/ActorModifierCore/Source/ActorModifierCore/Private/Modifiers/ActorModifierCoreBase.cpp
--- a/Engine/Plugins/VirtualProduction/ActorModifierCore/Source/ActorModifierCore/Private/Modifiers/ActorModifierCoreBase.cpp
+++ b/Engine/Plugins/VirtualProduction/ActorModifierCore/Source/ActorModifierCore/Private/Modifiers/ActorModifierCoreBase.cpp
@@ -449,8 +449,8 @@ void UActorModifierCoreBase::DeferInitializeModifier()
 	}
 
 	// Bind to world delegate, tick will be called when all actors have been loaded, unbind when actors have been loaded
-	FWorldDelegates::OnWorldPostActorTick.RemoveAll(this);
-	FWorldDelegates::OnWorldPostActorTick.AddUObject(this, &UActorModifierCoreBase::PostModifierWorldLoad);
+	FWorldDelegates::OnWorldTickEnd.RemoveAll(this);
+	FWorldDelegates::OnWorldTickEnd.AddUObject(this, &UActorModifierCoreBase::PostModifierWorldLoad);
 }
 
 UActorModifierCoreBase::UActorModifierCoreBase()
@@ -527,7 +527,12 @@ void UActorModifierCoreBase::PostDuplicate(EDuplicateMode::Type DuplicateMode)
 {
 	Super::PostDuplicate(DuplicateMode);
 
-	InitializeModifier(EActorModifierCoreEnableReason::Duplicate);
+	// Only consider non-pie duplication here.
+	// Initializing a modifier via PIE will be handled via UActorModifierCoreBase::PostModifierWorldLoad 
+	if (DuplicateMode != EDuplicateMode::PIE)
+	{
+		InitializeModifier(EActorModifierCoreEnableReason::Duplicate);
+	}
 }
 
 #if WITH_EDITOR
@@ -702,7 +707,7 @@ void UActorModifierCoreBase::PostModifierWorldLoad(UWorld* InWorld, ELevelTick I
 		}
 
 		// Remove handle
-		FWorldDelegates::OnWorldPostActorTick.RemoveAll(this);
+		FWorldDelegates::OnWorldTickEnd.RemoveAll(this);
 
 		// Initialize now that all actors of world have been post loaded
 		InitializeModifier(EActorModifierCoreEnableReason::Load);
diff --git a/Engine/Plugins/VirtualProduction/Avalanche/Source/AvalancheMask/Public/Mask2D/AvaMask2DBaseModifier.h b/Engine/Plugins/VirtualProduction/Avalanche/Source/AvalancheMask/Public/Mask2D/AvaMask2DBaseModifier.h
--- a/Engine/Plugins/VirtualProduction/Avalanche/Source/AvalancheMask/Public/Mask2D/AvaMask2DBaseModifier.h
+++ b/Engine/Plugins/VirtualProduction/Avalanche/Source/AvalancheMask/Public/Mask2D/AvaMask2DBaseModifier.h
@@ -221,7 +221,7 @@ protected:
 	FName AutoChannelName;
 
 	/** Cached actor data to apply/restore */
-	UPROPERTY(DuplicateTransient)
+	UPROPERTY(NonPIEDuplicateTransient)
 	TMap<TWeakObjectPtr<AActor>, FAvaMask2DActorData> ActorData;
 
 	/** Reference to the last resolved canvas name, as stored in CanvasWeak */
