Motion Design Actor Modifiers: deprecated LogModifier in favor of UE_LOG for performance.
Also added fixed unnecessary allocations in the masking system (e.g. copy instead of move, or using FindOrAdd single call instead of Find then Add).

diff --git a/Engine/Plugins/VirtualProduction/ActorModifierCore/Source/ActorModifierCore/Private/ActorModifierCoreModule.cpp b/Engine/Plugins/VirtualProduction/ActorModifierCore/Source/ActorModifierCore/Private/ActorModifierCoreModule.cpp
--- a/Engine/Plugins/VirtualProduction/ActorModifierCore/Source/ActorModifierCore/Private/ActorModifierCoreModule.cpp
+++ b/Engine/Plugins/VirtualProduction/ActorModifierCore/Source/ActorModifierCore/Private/ActorModifierCoreModule.cpp
@@ -1,5 +1,8 @@
 // Copyright Epic Games, Inc. All Rights Reserved.
 
 #include "Modules/ModuleManager.h"
+#include "ActorModifierCoreLog.h"
+
+DEFINE_LOG_CATEGORY(LogActorModifierCore)
 
 IMPLEMENT_MODULE(FDefaultModuleImpl, ActorModifierCore)
\ No newline at end of file
diff --git a/Engine/Plugins/VirtualProduction/ActorModifierCore/Source/ActorModifierCore/Private/Modifiers/ActorModifierCoreBase.cpp b/Engine/Plugins/VirtualProduction/ActorModifierCore/Source/ActorModifierCore/Private/Modifiers/ActorModifierCoreBase.cpp
--- a/Engine/Plugins/VirtualProduction/ActorModifierCore/Source/ActorModifierCore/Private/Modifiers/ActorModifierCoreBase.cpp
+++ b/Engine/Plugins/VirtualProduction/ActorModifierCore/Source/ActorModifierCore/Private/Modifiers/ActorModifierCoreBase.cpp
@@ -2,10 +2,12 @@
 
 #include "Modifiers/ActorModifierCoreBase.h"
 
+#include "ActorModifierCoreLog.h"
 #include "Async/TaskGraphInterfaces.h"
 #include "Modifiers/ActorModifierCoreComponent.h"
 #include "Modifiers/ActorModifierCoreSharedObject.h"
 #include "Modifiers/ActorModifierCoreStack.h"
+#include "Modifiers/Utilities/ActorModifierCoreUtilities.h"
 #include "Subsystems/ActorModifierCoreSubsystem.h"
 #include "UObject/Package.h"
 
@@ -34,7 +36,11 @@ UActorModifierCoreSharedObject* UActorModifierCoreBase::GetShared(TSubclassOf<UA
 
 	if (bInCreateIfNone && !SharedObject)
 	{
-		LogModifier(FString::Printf(TEXT("Failed to create and retrieve the specific shared object : %s"), *InClass.Get()->GetName()), /** Force */true, EActorModifierCoreStatus::Error);
+		UE_LOG(LogActorModifierCore, Error, TEXT("[%s][%s][%s] Failed to create and retrieve the specific shared object : %s")
+			, *UE::ActorModifierCore::Utilities::GetActorNameSafe(GetModifiedActor())
+			, *GetNameSafe(GetClass())
+			, *GetModifierName().ToString()
+			, *InClass.Get()->GetName());
 	}
 
 	return SharedObject;
@@ -60,8 +66,7 @@ void UActorModifierCoreBase::Next()
 			}
 			else
 			{
-				LogModifier(TEXT("Next is called again after execution is done"), true);
-				checkNoEntry()
+				checkf(false, TEXT("Next is called again after execution is done"))
 			}
 		}
 	};
@@ -104,8 +109,7 @@ void UActorModifierCoreBase::Fail(const FText& InFailReason)
 			}
 			else
 			{
-				LogModifier(TEXT("Fail is called again after execution is done"), true);
-				checkNoEntry()
+				checkf(false, TEXT("Fail is called again after execution is done"));
 			}
 		}
 	};
@@ -133,7 +137,10 @@ void UActorModifierCoreBase::Fail(const FText& InFailReason)
 
 void UActorModifierCoreBase::Unapply()
 {
-	LogModifier(TEXT("Unapplying modifier"));
+	UE_LOG(LogActorModifierCore, Verbose, TEXT("[%s][%s][%s] Unapplying modifier")
+		, *UE::ActorModifierCore::Utilities::GetActorNameSafe(GetModifiedActor())
+		, *GetNameSafe(GetClass())
+		, *GetModifierName().ToString());
 
 	// only restore if this modifier was already applied previously
 	if (bModifierApplied)
@@ -360,7 +367,10 @@ void UActorModifierCoreBase::LockModifierExecution()
 	{
 		bModifierExecutionLocked = true;
 
-		LogModifier(TEXT("Locking modifier execution"));
+		UE_LOG(LogActorModifierCore, Verbose, TEXT("[%s][%s][%s] Locking modifier execution")
+			, *UE::ActorModifierCore::Utilities::GetActorNameSafe(GetModifiedActor())
+			, *GetNameSafe(GetClass())
+			, *GetModifierName().ToString());
 	}
 }
 
@@ -370,7 +380,10 @@ void UActorModifierCoreBase::UnlockModifierExecution()
 	{
 		bModifierExecutionLocked = false;
 
-		LogModifier(TEXT("Unlocking modifier execution"));
+		UE_LOG(LogActorModifierCore, Verbose, TEXT("[%s][%s][%s] Unlocking modifier execution")
+			, *UE::ActorModifierCore::Utilities::GetActorNameSafe(GetModifiedActor())
+			, *GetNameSafe(GetClass())
+			, *GetModifierName().ToString());
 
 		if (IsModifierDirty())
 		{
@@ -386,7 +399,11 @@ void UActorModifierCoreBase::AddExtensionInternal(const FName& InExtensionType,
 		return;
 	}
 
-	LogModifier(FString::Printf(TEXT("Adding modifier extension %s"), *InExtensionType.ToString()));
+	UE_LOG(LogActorModifierCore, Verbose, TEXT("[%s][%s][%s] Adding modifier extension %s")
+		, *UE::ActorModifierCore::Utilities::GetActorNameSafe(GetModifiedActor())
+		, *GetNameSafe(GetClass())
+		, *GetModifierName().ToString()
+		, *InExtensionType.ToString());
 
 	ModifierExtensions.Emplace(InExtensionType, InExtension);
 
@@ -488,7 +505,14 @@ void UActorModifierCoreBase::PostLoad()
 				const FString ThisStackName = GetName();
 				const EObjectFlags ThisStackFlags = GetFlags();
 
-				LogModifier(FString::Printf(TEXT("Modifier stack migrated to component stack %s %s with %i modifiers"), *OwningComponent->GetName(), *ComponentStack->GetName(), ThisStack->Modifiers.Num()), true);
+				UE_LOG(LogActorModifierCore, Log, TEXT("[%s][%s][%s] Modifier stack migrated to component stack %s %s with %i modifiers")
+					, *UE::ActorModifierCore::Utilities::GetActorNameSafe(GetModifiedActor())
+					, *GetNameSafe(GetClass())
+					, *GetModifierName().ToString()
+					, *OwningComponent->GetName()
+					, *ComponentStack->GetName()
+					, ThisStack->Modifiers.Num());
+
 				Rename(nullptr, GetTransientPackage(), RenameFlags);
 
 				ComponentStack->Rename(*ThisStackName, nullptr, RenameFlags);
@@ -503,7 +527,13 @@ void UActorModifierCoreBase::PostLoad()
 		else
 		{
 			const bool bSuccess = Rename(nullptr, OwningComponent->GetModifierStack(), RenameFlags);
-			LogModifier(FString::Printf(TEXT("Modifier outer renamed to stack %s : %s"), *OwningComponent->GetModifierStack()->GetName(), bSuccess ? TEXT("OK") : TEXT("Fail")), true);
+
+			UE_LOG(LogActorModifierCore, Log, TEXT("[%s][%s][%s] Modifier outer renamed to stack %s : %s")
+				, *UE::ActorModifierCore::Utilities::GetActorNameSafe(GetModifiedActor())
+				, *GetNameSafe(GetClass())
+				, *GetModifierName().ToString()
+				, *OwningComponent->GetModifierStack()->GetName()
+				, bSuccess ? TEXT("OK") : TEXT("Fail"));
 		}
 	}
 
@@ -599,7 +629,11 @@ bool UActorModifierCoreBase::RemoveExtension(const FName& InExtensionType)
 {
 	if (const TSharedPtr<FActorModifierCoreExtension>* Extension = ModifierExtensions.Find(InExtensionType))
 	{
-		LogModifier(FString::Printf(TEXT("Removing modifier extension %s"), *InExtensionType.ToString()));
+		UE_LOG(LogActorModifierCore, Verbose, TEXT("[%s][%s][%s] Removing modifier extension %s")
+			, *UE::ActorModifierCore::Utilities::GetActorNameSafe(GetModifiedActor())
+			, *GetNameSafe(GetClass())
+			, *GetModifierName().ToString()
+			, *InExtensionType.ToString());
 
 		if (bModifierEnabled)
 		{
@@ -614,25 +648,6 @@ bool UActorModifierCoreBase::RemoveExtension(const FName& InExtensionType)
 
 void UActorModifierCoreBase::LogModifier(const FString& InLog, bool bInForce, EActorModifierCoreStatus InType) const
 {
-	if (IsModifierProfiling() || bInForce)
-	{
-		const FString ActorLabel = GetModifiedActor() ? *GetModifiedActor()->GetActorNameOrLabel() : TEXT("Invalid actor");
-		const FString ModifierLabel = GetModifierName().ToString();
-		const FString ClassLabel = GetClass()->GetName();
-
-		switch (InType)
-		{
-		case EActorModifierCoreStatus::Success:
-			UE_LOG(LogActorModifierCoreBase, Log, TEXT("[%s][%s][%s] %s"), *ActorLabel, *ClassLabel, *ModifierLabel, *InLog);
-			break;
-		case EActorModifierCoreStatus::Warning:
-			UE_LOG(LogActorModifierCoreBase, Warning, TEXT("[%s][%s][%s] %s"), *ActorLabel, *ClassLabel, *ModifierLabel, *InLog);
-			break;
-		case EActorModifierCoreStatus::Error:
-			UE_LOG(LogActorModifierCoreBase, Error, TEXT("[%s][%s][%s] %s"), *ActorLabel, *ClassLabel, *ModifierLabel, *InLog);
-			break;
-		}
-	}
 }
 
 void UActorModifierCoreBase::PostModifierCreation(UActorModifierCoreStack* InStack)
@@ -789,7 +804,11 @@ void UActorModifierCoreBase::InitializeModifier(EActorModifierCoreEnableReason I
 			Profiler = Metadata->CreateProfilerInstance(this);
 		}
 
-		LogModifier(FString::Printf(TEXT("Initializing modifier with reason %i"), InReason));
+		UE_LOG(LogActorModifierCore, Verbose, TEXT("[%s][%s][%s] Initializing modifier with reason %i")
+			, *UE::ActorModifierCore::Utilities::GetActorNameSafe(GetModifiedActor())
+			, *GetNameSafe(GetClass())
+			, *GetModifierName().ToString()
+			, InReason);
 
 		// add the modifier to our new actor stack
 		OnModifierAdded(InReason);
@@ -813,7 +832,11 @@ void UActorModifierCoreBase::UninitializeModifier(EActorModifierCoreDisableReaso
 	{
 		bModifierInitialized = false;
 
-		LogModifier(FString::Printf(TEXT("Uninitializing modifier with reason %i"), InReason));
+		UE_LOG(LogActorModifierCore, Verbose, TEXT("[%s][%s][%s] Uninitializing modifier with reason %i")
+			, *UE::ActorModifierCore::Utilities::GetActorNameSafe(GetModifiedActor())
+			, *GetNameSafe(GetClass())
+			, *GetModifierName().ToString()
+			, InReason);
 
 		const bool bWasModifierEnabled = bModifierEnabled;
 
@@ -841,7 +864,11 @@ void UActorModifierCoreBase::UninitializeModifier(EActorModifierCoreDisableReaso
 
 void UActorModifierCoreBase::OnModifierEnabledChanged(bool bInExecute)
 {
-	LogModifier(FString::Printf(TEXT("Modifier %s"), bModifierEnabled ? TEXT("enabled") : TEXT("disabled")), true);
+	UE_LOG(LogActorModifierCore, Verbose, TEXT("[%s][%s][%s] Modifier %s")
+		, *UE::ActorModifierCore::Utilities::GetActorNameSafe(GetModifiedActor())
+		, *GetNameSafe(GetClass())
+		, *GetModifierName().ToString()
+		, bModifierEnabled ? TEXT("enabled") : TEXT("disabled"));
 
 	if (bModifierEnabled)
 	{
@@ -868,7 +895,10 @@ void UActorModifierCoreBase::BeginModifierExecution()
 		Profiler->BeginProfiling();
 	}
 
-	LogModifier(TEXT("Appling modifier"));
+	UE_LOG(LogActorModifierCore, Verbose, TEXT("[%s][%s][%s] Applying Modifier")
+		, *UE::ActorModifierCore::Utilities::GetActorNameSafe(GetModifiedActor())
+		, *GetNameSafe(GetClass())
+		, *GetModifierName().ToString());
 }
 
 void UActorModifierCoreBase::EndModifierExecution()
@@ -884,9 +914,10 @@ void UActorModifierCoreBase::EndModifierExecution()
 		Profiler->EndProfiling();
 	}
 
-	if (Status.GetStatus() != EActorModifierCoreStatus::Success)
-	{
-		LogModifier(FString::Printf(TEXT("Modifier execution failed due to reason : %s"), *Status.GetStatusMessage().ToString()), true);
-	}
+	UE_CLOG(Status.GetStatus() != EActorModifierCoreStatus::Success, LogActorModifierCore, Warning, TEXT("[%s][%s][%s] Modifier execution failed due to reason : %s")
+		, *UE::ActorModifierCore::Utilities::GetActorNameSafe(GetModifiedActor())
+		, *GetNameSafe(GetClass())
+		, *GetModifierName().ToString()
+		, *Status.GetStatusMessage().ToString());
 }
 
diff --git a/Engine/Plugins/VirtualProduction/ActorModifierCore/Source/ActorModifierCore/Private/Modifiers/ActorModifierCoreStack.cpp b/Engine/Plugins/VirtualProduction/ActorModifierCore/Source/ActorModifierCore/Private/Modifiers/ActorModifierCoreStack.cpp
--- a/Engine/Plugins/VirtualProduction/ActorModifierCore/Source/ActorModifierCore/Private/Modifiers/ActorModifierCoreStack.cpp
+++ b/Engine/Plugins/VirtualProduction/ActorModifierCore/Source/ActorModifierCore/Private/Modifiers/ActorModifierCoreStack.cpp
@@ -2,10 +2,12 @@
 
 #include "Modifiers/ActorModifierCoreStack.h"
 
+#include "ActorModifierCoreLog.h"
 #include "Algo/RemoveIf.h"
 #include "Components/SceneComponent.h"
 #include "Modifiers/ActorModifierCoreComponent.h"
 #include "Modifiers/Blueprints/ActorModifierCoreBlueprintBase.h"
+#include "Modifiers/Utilities/ActorModifierCoreUtilities.h"
 #include "Subsystems/ActorModifierCoreSubsystem.h"
 
 #define LOCTEXT_NAMESPACE "ActorModifierCoreStack"
@@ -85,7 +87,12 @@ bool UActorModifierCoreStack::IsModifierStackInitialized() const
 void UActorModifierCoreStack::SetModifierProfiling(bool bInProfiling)
 {
 	bModifierProfiling = bInProfiling;
-	LogModifier(FString::Printf(TEXT("Profiling %s"), bInProfiling ? TEXT("enabled") : TEXT("disabled")), true);
+
+	UE_LOG(LogActorModifierCore, Verbose, TEXT("[%s][%s][%s] Profiling %s")
+		, *UE::ActorModifierCore::Utilities::GetActorNameSafe(GetModifiedActor())
+		, *GetNameSafe(GetClass())
+		, *GetModifierName().ToString()
+		, bInProfiling ? TEXT("enabled") : TEXT("disabled"));
 }
 
 void UActorModifierCoreStack::PostLoad()
@@ -659,7 +666,11 @@ UActorModifierCoreBase* UActorModifierCoreStack::CloneModifier(FActorModifierCor
 		}
 	}
 
-	LogModifier(FString::Printf(TEXT("Clone modifier %s"), *InCloneOp.CloneModifier->GetModifierName().ToString()), true);
+	UE_LOG(LogActorModifierCore, Verbose, TEXT("[%s][%s][%s] Clone modifier %s")
+		, *UE::ActorModifierCore::Utilities::GetActorNameSafe(GetModifiedActor())
+		, *GetNameSafe(GetClass())
+		, *GetModifierName().ToString()
+		, *InCloneOp.CloneModifier->GetModifierName().ToString());
 
 	// Refresh the stack
 	MarkModifierDirty();
@@ -838,7 +849,11 @@ UActorModifierCoreBase* UActorModifierCoreStack::InsertModifier(FActorModifierCo
 		}
 	}
 
-	LogModifier(FString::Printf(TEXT("Insert modifier %s"), *InInsertOp.NewModifierName.ToString()), true);
+	UE_LOG(LogActorModifierCore, Verbose, TEXT("[%s][%s][%s] Insert modifier %s")
+		, *UE::ActorModifierCore::Utilities::GetActorNameSafe(GetModifiedActor())
+		, *GetNameSafe(GetClass())
+		, *GetModifierName().ToString()
+		, *InInsertOp.NewModifierName.ToString());
 
 	// Refresh the stack
 	MarkModifierDirty();
@@ -1097,7 +1112,11 @@ bool UActorModifierCoreStack::MoveModifier(FActorModifierCoreStackMoveOp& InMove
 
 	Modifiers.Insert(InMoveOp.MoveModifier, MoveIdx);
 
-	LogModifier(FString::Printf(TEXT("Move modifier %s"), *InMoveOp.MoveModifier->GetModifierName().ToString()), true);
+	UE_LOG(LogActorModifierCore, Verbose, TEXT("[%s][%s][%s] Move modifier %s")
+		, *UE::ActorModifierCore::Utilities::GetActorNameSafe(GetModifiedActor())
+		, *GetNameSafe(GetClass())
+		, *GetModifierName().ToString()
+		, *InMoveOp.MoveModifier->GetModifierName().ToString());
 
 	OnModifierMovedDelegate.Broadcast(InMoveOp.MoveModifier);
 
@@ -1211,7 +1230,11 @@ bool UActorModifierCoreStack::RemoveModifier(FActorModifierCoreStackRemoveOp& In
 		}
 	}
 
-	LogModifier(FString::Printf(TEXT("Remove modifier %s"), *InRemoveOp.RemoveModifier->GetModifierName().ToString()), true);
+	UE_LOG(LogActorModifierCore, Verbose, TEXT("[%s][%s][%s] Remove modifier %s")
+		, *UE::ActorModifierCore::Utilities::GetActorNameSafe(GetModifiedActor())
+		, *GetNameSafe(GetClass())
+		, *GetModifierName().ToString()
+		, *InRemoveOp.RemoveModifier->GetModifierName().ToString());
 
 	// update the stack to restore to initial state
 	MarkModifierDirty();
@@ -1517,7 +1540,12 @@ void UActorModifierCoreStack::CheckModifierOptimization(bool bInInvalidateAll)
 					// Only log when we invalidate all and recheck optimization
 					if (bInInvalidateAll)
 					{
-						LogModifier(FString::Printf(TEXT("Optimisation possible for modifier %s : %s"), *Modifier->GetModifierName().ToString(), *Message.ToString()), true);
+						UE_LOG(LogActorModifierCore, Log, TEXT("[%s][%s][%s] Optimization possible for modifier %s : %s")
+							, *UE::ActorModifierCore::Utilities::GetActorNameSafe(GetModifiedActor())
+							, *GetNameSafe(GetClass())
+							, *GetModifierName().ToString()
+							, *Modifier->GetModifierName().ToString()
+							, *Message.ToString());
 					}
 
 					break;
@@ -1547,7 +1575,12 @@ void UActorModifierCoreStack::CheckModifierOptimization(bool bInInvalidateAll)
 					// Only log when we invalidate all and recheck optimization
 					if (bInInvalidateAll)
 					{
-						LogModifier(FString::Printf(TEXT("Optimisation possible for modifier %s : %s"), *Modifier->GetModifierName().ToString(), *Message.ToString()), true);
+						UE_LOG(LogActorModifierCore, Log, TEXT("[%s][%s][%s] Optimization possible for modifier %s : %s")
+							, *UE::ActorModifierCore::Utilities::GetActorNameSafe(GetModifiedActor())
+							, *GetNameSafe(GetClass())
+							, *GetModifierName().ToString()
+							, *Modifier->GetModifierName().ToString()
+							, *Message.ToString());
 					}
 
 					break;
@@ -1704,7 +1737,11 @@ void UActorModifierCoreStack::OnModifierDirty(UActorModifierCoreBase* DirtyModif
 		return;
 	}
 
-	LogModifier(FString::Printf(TEXT("Modifier %s marked dirty"), *DirtyModifier->GetModifierName().ToString()));
+	UE_LOG(LogActorModifierCore, Verbose, TEXT("[%s][%s][%s] Modifier %s marked dirty")
+		, *UE::ActorModifierCore::Utilities::GetActorNameSafe(GetModifiedActor())
+		, *GetNameSafe(GetClass())
+		, *GetModifierName().ToString()
+		, *DirtyModifier->GetModifierName().ToString());
 
 	if (!bExecute
 		|| DirtyModifier->IsModifierExecutionLocked()
@@ -1717,7 +1754,10 @@ void UActorModifierCoreStack::OnModifierDirty(UActorModifierCoreBase* DirtyModif
 		return;
 	}
 
-	LogModifier(TEXT("Executing modifier root stack"));
+	UE_LOG(LogActorModifierCore, Verbose, TEXT("[%s][%s][%s] Executing modifier root stack")
+		, *UE::ActorModifierCore::Utilities::GetActorNameSafe(GetModifiedActor())
+		, *GetNameSafe(GetClass())
+		, *GetModifierName().ToString());
 
 	// Execute root stack as standalone modifier
 	RootExecutionTask = FActorModifierCoreExecutionTask();
diff --git a/Engine/Plugins/VirtualProduction/ActorModifierCore/Source/ActorModifierCore/Private/Modifiers/Blueprints/ActorModifierCoreBlueprintBase.cpp b/Engine/Plugins/VirtualProduction/ActorModifierCore/Source/ActorModifierCore/Private/Modifiers/Blueprints/ActorModifierCoreBlueprintBase.cpp
--- a/Engine/Plugins/VirtualProduction/ActorModifierCore/Source/ActorModifierCore/Private/Modifiers/Blueprints/ActorModifierCoreBlueprintBase.cpp
+++ b/Engine/Plugins/VirtualProduction/ActorModifierCore/Source/ActorModifierCore/Private/Modifiers/Blueprints/ActorModifierCoreBlueprintBase.cpp
@@ -1,6 +1,8 @@
 // Copyright Epic Games, Inc. All Rights Reserved.
 
 #include "Modifiers/Blueprints/ActorModifierCoreBlueprintBase.h"
+#include "ActorModifierCoreLog.h"
+#include "Modifiers/Utilities/ActorModifierCoreUtilities.h"
 
 #define LOCTEXT_NAMESPACE "ActorModifierCoreBlueprintBase"
 
@@ -8,7 +10,10 @@ void UActorModifierCoreBlueprintBase::FlagModifierDirty()
 {
 	if (!IsTemplate())
 	{
-		LogModifier(TEXT("Blueprint modifier flagged dirty"));
+		UE_LOG(LogActorModifierCore, Verbose, TEXT("[%s][%s][%s] Blueprint modifier flagged dirty")
+			, *UE::ActorModifierCore::Utilities::GetActorNameSafe(GetModifiedActor())
+			, *GetNameSafe(GetClass())
+			, *GetModifierName().ToString());
 
 		MarkModifierDirty(/** Execute */true);
 	}
@@ -34,14 +39,20 @@ void UActorModifierCoreBlueprintBase::OnModifierCDOSetup(FActorModifierCoreMetad
 {
 	Super::OnModifierCDOSetup(InMetadata);
 
-	LogModifier(TEXT("Blueprint modifier setup"));
+	UE_LOG(LogActorModifierCore, Verbose, TEXT("[%s][%s][%s] Blueprint modifier setup")
+		, *UE::ActorModifierCore::Utilities::GetActorNameSafe(GetModifiedActor())
+		, *GetNameSafe(GetClass())
+		, *GetModifierName().ToString());
 
 	// Since the blueprint utility functions returns a copy of the struct, make sure we get a result
 	OnModifierSetupEvent(InMetadata, /** OutMetadata */InMetadata);
 
 	if (InMetadata.GetName().IsNone())
 	{
-		LogModifier(TEXT("Blueprint modifier setup failed : Name was not defined"), true);
+		UE_LOG(LogActorModifierCore, Warning, TEXT("[%s][%s][%s] Blueprint modifier setup failed : Name was not defined")
+			, *UE::ActorModifierCore::Utilities::GetActorNameSafe(GetModifiedActor())
+			, *GetNameSafe(GetClass())
+			, *GetModifierName().ToString());
 	}
 }
 
@@ -53,7 +64,11 @@ void UActorModifierCoreBlueprintBase::OnModifierAdded(EActorModifierCoreEnableRe
 
 	if (IsValid(TargetActor))
 	{
-		LogModifier(FString::Printf(TEXT("Blueprint modifier added with reason %i"), InReason));
+		UE_LOG(LogActorModifierCore, Verbose, TEXT("[%s][%s][%s] Blueprint modifier added with reason %i")
+			, *UE::ActorModifierCore::Utilities::GetActorNameSafe(GetModifiedActor())
+			, *GetNameSafe(GetClass())
+			, *GetModifierName().ToString()
+			, InReason);
 
 		OnModifierAddedEvent(TargetActor, InReason);
 	}
@@ -67,7 +82,11 @@ void UActorModifierCoreBlueprintBase::OnModifierEnabled(EActorModifierCoreEnable
 
 	if (IsValid(TargetActor))
 	{
-		LogModifier(FString::Printf(TEXT("Blueprint modifier enabled with reason %i"), InReason));
+		UE_LOG(LogActorModifierCore, Verbose, TEXT("[%s][%s][%s] Blueprint modifier enabled with reason %i")
+			, *UE::ActorModifierCore::Utilities::GetActorNameSafe(GetModifiedActor())
+			, *GetNameSafe(GetClass())
+			, *GetModifierName().ToString()
+			, InReason);
 
 		OnModifierEnabledEvent(TargetActor, InReason);
 	}
@@ -81,7 +100,11 @@ void UActorModifierCoreBlueprintBase::OnModifierDisabled(EActorModifierCoreDisab
 
 	if (IsValid(TargetActor))
 	{
-		LogModifier(FString::Printf(TEXT("Blueprint modifier disabled with reason %i"), InReason));
+		UE_LOG(LogActorModifierCore, Verbose, TEXT("[%s][%s][%s] Blueprint modifier disabled with reason %i")
+			, *UE::ActorModifierCore::Utilities::GetActorNameSafe(GetModifiedActor())
+			, *GetNameSafe(GetClass())
+			, *GetModifierName().ToString()
+			, InReason);
 
 		OnModifierDisabledEvent(TargetActor, InReason);
 	}
@@ -95,7 +118,11 @@ void UActorModifierCoreBlueprintBase::OnModifierRemoved(EActorModifierCoreDisabl
 
 	if (IsValid(TargetActor))
 	{
-		LogModifier(FString::Printf(TEXT("Blueprint modifier removed with reason %i"), InReason));
+		UE_LOG(LogActorModifierCore, Verbose, TEXT("[%s][%s][%s] Blueprint modifier removed with reason %i")
+			, *UE::ActorModifierCore::Utilities::GetActorNameSafe(GetModifiedActor())
+			, *GetNameSafe(GetClass())
+			, *GetModifierName().ToString()
+			, InReason);
 
 		OnModifierRemovedEvent(TargetActor, InReason);
 	}
@@ -109,7 +136,10 @@ void UActorModifierCoreBlueprintBase::SavePreState()
 
 	if (IsValid(TargetActor))
 	{
-		LogModifier(TEXT("Blueprint modifier save pre state"));
+		UE_LOG(LogActorModifierCore, Verbose, TEXT("[%s][%s][%s] Blueprint modifier save pre state")
+			, *UE::ActorModifierCore::Utilities::GetActorNameSafe(GetModifiedActor())
+			, *GetNameSafe(GetClass())
+			, *GetModifierName().ToString());
 
 		OnModifierSaveStateEvent(TargetActor);
 	}
@@ -123,7 +153,10 @@ void UActorModifierCoreBlueprintBase::RestorePreState()
 
 	if (IsValid(TargetActor))
 	{
-		LogModifier(TEXT("Blueprint modifier restore pre state"));
+		UE_LOG(LogActorModifierCore, Verbose, TEXT("[%s][%s][%s] Blueprint modifier restore pre state")
+			, *UE::ActorModifierCore::Utilities::GetActorNameSafe(GetModifiedActor())
+			, *GetNameSafe(GetClass())
+			, *GetModifierName().ToString());
 
 		OnModifierRestoreStateEvent(TargetActor);
 	}
@@ -134,7 +167,10 @@ void UActorModifierCoreBlueprintBase::Apply()
 	FText FailReason = FText::GetEmpty();
 	AActor* TargetActor = GetModifiedActor();
 
-	LogModifier(TEXT("Blueprint modifier apply"));
+	UE_LOG(LogActorModifierCore, Verbose, TEXT("[%s][%s][%s] Blueprint modifier apply")
+		, *UE::ActorModifierCore::Utilities::GetActorNameSafe(GetModifiedActor())
+		, *GetNameSafe(GetClass())
+		, *GetModifierName().ToString());
 
 	if (IsValid(TargetActor) && OnModifierApplyEvent(TargetActor, FailReason))
 	{
diff --git a/Engine/Plugins/VirtualProduction/ActorModifierCore/Source/ActorModifierCore/Private/Modifiers/Utilities/ActorModifierCoreUtilities.cpp b/Engine/Plugins/VirtualProduction/ActorModifierCore/Source/ActorModifierCore/Private/Modifiers/Utilities/ActorModifierCoreUtilities.cpp
--- a/Engine/Plugins/VirtualProduction/ActorModifierCore/Source/ActorModifierCore/Private/Modifiers/Utilities/ActorModifierCoreUtilities.cpp
+++ b/Engine/Plugins/VirtualProduction/ActorModifierCore/Source/ActorModifierCore/Private/Modifiers/Utilities/ActorModifierCoreUtilities.cpp
@@ -5,7 +5,10 @@
 #include "Modifiers/ActorModifierCoreComponent.h"
 #include "Modifiers/ActorModifierCoreStack.h"
 
-UActorModifierCoreBase* UE::ActorModifierCore::Utilities::FindFirstActorModifierByClass(const AActor* InStartActor, const TSubclassOf<UActorModifierCoreBase>& InModifierClass)
+namespace UE::ActorModifierCore::Utilities
+{
+
+UActorModifierCoreBase* FindFirstActorModifierByClass(const AActor* InStartActor, const TSubclassOf<UActorModifierCoreBase>& InModifierClass)
 {
 	UActorModifierCoreBase* FoundModifier = nullptr;
 
@@ -29,3 +32,14 @@ UActorModifierCoreBase* UE::ActorModifierCore::Utilities::FindFirstActorModifier
 
 	return FindFirstActorModifierByClass(InStartActor->GetAttachParentActor(), InModifierClass);
 }
+
+FString GetActorNameSafe(const AActor* InActor)
+{
+	if (InActor)
+	{
+		return InActor->GetActorNameOrLabel();
+	}
+	return TEXT("None");
+}
+
+} // UE::ActorModifierCore::Utilities
diff --git a/Engine/Plugins/VirtualProduction/ActorModifierCore/Source/ActorModifierCore/Public/ActorModifierCoreLog.h b/Engine/Plugins/VirtualProduction/ActorModifierCore/Source/ActorModifierCore/Public/ActorModifierCoreLog.h
new file mode 100644
--- /dev/null
+++ b/Engine/Plugins/VirtualProduction/ActorModifierCore/Source/ActorModifierCore/Public/ActorModifierCoreLog.h
@@ -0,0 +1,7 @@
+// Copyright Epic Games, Inc. All Rights Reserved.
+
+#pragma once
+
+#include "Logging/LogMacros.h"
+
+DECLARE_LOG_CATEGORY_EXTERN(LogActorModifierCore, Log, All);
diff --git a/Engine/Plugins/VirtualProduction/ActorModifierCore/Source/ActorModifierCore/Public/Modifiers/ActorModifierCoreBase.h b/Engine/Plugins/VirtualProduction/ActorModifierCore/Source/ActorModifierCore/Public/Modifiers/ActorModifierCoreBase.h
--- a/Engine/Plugins/VirtualProduction/ActorModifierCore/Source/ActorModifierCore/Public/Modifiers/ActorModifierCoreBase.h
+++ b/Engine/Plugins/VirtualProduction/ActorModifierCore/Source/ActorModifierCore/Public/Modifiers/ActorModifierCoreBase.h
@@ -351,7 +351,7 @@ protected:
 	/** Removes an extension of this type from this modifier */
 	ACTORMODIFIERCORE_API bool RemoveExtension(const FName& InExtensionType);
 
-	/** Logs modifier message if in profiling mode or forced */
+	UE_DEPRECATED(5.8, "LogModifier deprecated. Prefer using UE_LOG instead.")
 	ACTORMODIFIERCORE_API void LogModifier(const FString& InLog, bool bInForce = false, EActorModifierCoreStatus InType = EActorModifierCoreStatus::Success) const;
 
 	void DeferInitializeModifier();
diff --git a/Engine/Plugins/VirtualProduction/ActorModifierCore/Source/ActorModifierCore/Public/Modifiers/Utilities/ActorModifierCoreUtilities.h b/Engine/Plugins/VirtualProduction/ActorModifierCore/Source/ActorModifierCore/Public/Modifiers/Utilities/ActorModifierCoreUtilities.h
--- a/Engine/Plugins/VirtualProduction/ActorModifierCore/Source/ActorModifierCore/Public/Modifiers/Utilities/ActorModifierCoreUtilities.h
+++ b/Engine/Plugins/VirtualProduction/ActorModifierCore/Source/ActorModifierCore/Public/Modifiers/Utilities/ActorModifierCoreUtilities.h
@@ -8,4 +8,7 @@ namespace UE::ActorModifierCore::Utilities
 {
 	/** Climbs up the attachment tree to find a modifier with a specific class */
 	ACTORMODIFIERCORE_API UActorModifierCoreBase* FindFirstActorModifierByClass(const AActor* InStartActor, const TSubclassOf<UActorModifierCoreBase>& InModifierClass);
+
+	/** Safely gets the actor label or name (if label is empty), returning an empty string if actor is null */
+	ACTORMODIFIERCORE_API FString GetActorNameSafe(const AActor* InActor);
 }
diff --git a/Engine/Plugins/VirtualProduction/Avalanche/Source/AvalancheMask/Private/Handling/IAvaMaskMaterialCollectionHandle.h b/Engine/Plugins/VirtualProduction/Avalanche/Source/AvalancheMask/Private/Handling/IAvaMaskMaterialCollectionHandle.h
--- a/Engine/Plugins/VirtualProduction/Avalanche/Source/AvalancheMask/Private/Handling/IAvaMaskMaterialCollectionHandle.h
+++ b/Engine/Plugins/VirtualProduction/Avalanche/Source/AvalancheMask/Private/Handling/IAvaMaskMaterialCollectionHandle.h
@@ -339,9 +339,8 @@ bool TAvaMaskMaterialCollectionHandle<HandleDataType>::ValidateMaterials(FText&
 		TArray<FString> MissingParameterNames;
 		if (!InMaterialHandle->HasRequiredParameters(MissingParameterNames))
 		{
-			FailedMaterialNames.Emplace(InMaterialHandle->GetMaterialName(), MissingParameterNames);
+			FailedMaterialNames.Emplace(InMaterialHandle->GetMaterialName(), MoveTemp(MissingParameterNames));
 			bAllMaterialsValid = false;
-			
 			return false;
 		}
 
diff --git a/Engine/Plugins/VirtualProduction/Avalanche/Source/AvalancheMask/Private/Mask2D/AvaMask2DReadModifier.cpp b/Engine/Plugins/VirtualProduction/Avalanche/Source/AvalancheMask/Private/Mask2D/AvaMask2DReadModifier.cpp
--- a/Engine/Plugins/VirtualProduction/Avalanche/Source/AvalancheMask/Private/Mask2D/AvaMask2DReadModifier.cpp
+++ b/Engine/Plugins/VirtualProduction/Avalanche/Source/AvalancheMask/Private/Mask2D/AvaMask2DReadModifier.cpp
@@ -151,8 +151,14 @@ bool UAvaMask2DReadModifier::ApplyRead(AActor* InActor, FAvaMask2DActorData& InA
 			ApplyParameters.MaterialParameters.OuterFeatherRadius = OuterFeatherRadius;
 			ApplyParameters.MaterialParameters.InnerFeatherRadius = InnerFeatherRadius;
 
-			FInstancedStruct& MaterialCollectionData = MaterialCollectionHandleData.FindOrAdd(InActor, MaterialCollectionHandle->MakeDataStruct());
-			MaterialCollectionHandle->ApplyModifiedState(ApplyParameters, MaterialCollectionData);
+			FInstancedStruct* MaterialCollectionData = MaterialCollectionHandleData.Find(InActor);
+			if (!MaterialCollectionData)
+			{
+				MaterialCollectionData = &MaterialCollectionHandleData.Add(InActor, MaterialCollectionHandle->MakeDataStruct());
+			}
+			check(MaterialCollectionData);
+
+			MaterialCollectionHandle->ApplyModifiedState(ApplyParameters, *MaterialCollectionData);
 		}
 
 		if (UGeometryMaskReadComponent* ReadComponent = FindOrAddMaskComponent<UGeometryMaskReadComponent>(InActor))
diff --git a/Engine/Plugins/VirtualProduction/Avalanche/Source/AvalancheMask/Private/Materials/AvaMaskMaterialInstanceSubsystem.cpp b/Engine/Plugins/VirtualProduction/Avalanche/Source/AvalancheMask/Private/Materials/AvaMaskMaterialInstanceSubsystem.cpp
--- a/Engine/Plugins/VirtualProduction/Avalanche/Source/AvalancheMask/Private/Materials/AvaMaskMaterialInstanceSubsystem.cpp
+++ b/Engine/Plugins/VirtualProduction/Avalanche/Source/AvalancheMask/Private/Materials/AvaMaskMaterialInstanceSubsystem.cpp
@@ -18,9 +18,15 @@ UMaterialInstanceDynamic* UAvaMaskMaterialInstanceProvider::FindOrAddMID(
 	{
 		return nullptr;
 	}
-	
-	FAvaMaskMaterialPermutations& MaterialPermutations = MaterialInstanceCache.FindOrAdd(InParentMaterial, { InParentMaterial, {} });
-	if (const TWeakObjectPtr<UMaterialInstanceDynamic>* WeakMaterialInstance = MaterialPermutations.Instances.Find(InInstanceKey))
+
+	FAvaMaskMaterialPermutations* MaterialPermutations = MaterialInstanceCache.Find(InParentMaterial);
+	if (!MaterialPermutations)
+	{
+		MaterialPermutations = &MaterialInstanceCache.Add(InParentMaterial, { InParentMaterial, {} });
+	}
+	check(MaterialPermutations);
+
+	if (const TWeakObjectPtr<UMaterialInstanceDynamic>* WeakMaterialInstance = MaterialPermutations->Instances.Find(InInstanceKey))
 	{
 		if (WeakMaterialInstance->IsValid())
 		{
@@ -35,7 +41,7 @@ UMaterialInstanceDynamic* UAvaMaskMaterialInstanceProvider::FindOrAddMID(
 		if (UMaterialInstanceDynamic* ResultFromParentProvider = Parent->FindOrAddMID(InParentMaterial, InInstanceKey))
 		{
 			// If found in parent, duplicate with this as outer
-			MaterialPermutations.Instances.Emplace(InInstanceKey, DuplicateObject(ResultFromParentProvider, this));
+			MaterialPermutations->Instances.Emplace(InInstanceKey, DuplicateObject(ResultFromParentProvider, this));
 			return ResultFromParentProvider;
 		}
 	}
@@ -58,7 +64,7 @@ UMaterialInstanceDynamic* UAvaMaskMaterialInstanceProvider::FindOrAddMID(
 		MaterialInstance = UMaterialInstanceDynamic::Create(InParentMaterial, this);
 	}
 
-	MaterialPermutations.Instances.Emplace(InInstanceKey, MaterialInstance);
+	MaterialPermutations->Instances.Emplace(InInstanceKey, MaterialInstance);
 
 	return MaterialInstance;
 }
diff --git a/Engine/Plugins/VirtualProduction/Avalanche/Source/AvalancheModifiers/Private/AvaModifiersModule.cpp b/Engine/Plugins/VirtualProduction/Avalanche/Source/AvalancheModifiers/Private/AvaModifiersModule.cpp
--- a/Engine/Plugins/VirtualProduction/Avalanche/Source/AvalancheModifiers/Private/AvaModifiersModule.cpp
+++ b/Engine/Plugins/VirtualProduction/Avalanche/Source/AvalancheModifiers/Private/AvaModifiersModule.cpp
@@ -1,13 +1,16 @@
-﻿// Copyright Epic Games, Inc. All Rights Reserved.
+// Copyright Epic Games, Inc. All Rights Reserved.
 
 #include "AvaModifiersModule.h"
 
+#include "AvaModifiersLog.h"
 #include "AvaModifiersSceneTreeResolver.h"
 #include "Engine/Level.h"
 #include "Extensions/ActorModifierSceneTreeUpdateExtension.h"
 #include "GameFramework/Actor.h"
 #include "Modules/ModuleManager.h"
 
+DEFINE_LOG_CATEGORY(LogAvaModifier);
+
 IMPLEMENT_MODULE(FAvalancheModifiersModule, AvalancheModifiers)
 
 void FAvalancheModifiersModule::StartupModule()
diff --git a/Engine/Plugins/VirtualProduction/Avalanche/Source/AvalancheModifiers/Private/Modifiers/AvaGlobalOpacityModifier.cpp b/Engine/Plugins/VirtualProduction/Avalanche/Source/AvalancheModifiers/Private/Modifiers/AvaGlobalOpacityModifier.cpp
--- a/Engine/Plugins/VirtualProduction/Avalanche/Source/AvalancheModifiers/Private/Modifiers/AvaGlobalOpacityModifier.cpp
+++ b/Engine/Plugins/VirtualProduction/Avalanche/Source/AvalancheModifiers/Private/Modifiers/AvaGlobalOpacityModifier.cpp
@@ -2,11 +2,13 @@
 
 #include "Modifiers/AvaGlobalOpacityModifier.h"
 
+#include "AvaModifiersLog.h"
 #include "Components/DMMaterialValue.h"
 #include "Components/MaterialValues/DMMaterialValueFloat1.h"
 #include "Material/DynamicMaterialInstance.h"
 #include "Materials/Material.h"
 #include "Model/DynamicMaterialModel.h"
+#include "Modifiers/Utilities/ActorModifierCoreUtilities.h"
 
 #define LOCTEXT_NAMESPACE "AvaGlobalOpacityModifier"
 
@@ -70,7 +72,10 @@ void UAvaGlobalOpacityModifier::OnGlobalOpacityChanged()
 	}
 	else
 	{
-		LogModifier(TEXT("Unable to find or add the global opacity material parameter."), /* Force */ false, EActorModifierCoreStatus::Error);
+		UE_LOG(LogAvaModifier, Warning, TEXT("[%s][%s][%s] Unable to find or add the global opacity material parameter")
+			, *UE::ActorModifierCore::Utilities::GetActorNameSafe(GetModifiedActor())
+			, *GetNameSafe(GetClass())
+			, *GetModifierName().ToString());
 	}
 }
 
diff --git a/Engine/Plugins/VirtualProduction/Avalanche/Source/AvalancheModifiers/Public/AvaModifiersLog.h b/Engine/Plugins/VirtualProduction/Avalanche/Source/AvalancheModifiers/Public/AvaModifiersLog.h
new file mode 100644
--- /dev/null
+++ b/Engine/Plugins/VirtualProduction/Avalanche/Source/AvalancheModifiers/Public/AvaModifiersLog.h
@@ -0,0 +1,7 @@
+// Copyright Epic Games, Inc. All Rights Reserved.
+
+#pragma once
+
+#include "Logging/LogMacros.h"
+
+DECLARE_LOG_CATEGORY_EXTERN(LogAvaModifier, Log, All);
