Motion Design: corrected log verbosity from Log to Verbose for logs that were constantly outputting for any valid scene.

#13623
#Synced-CL 48355279

diff --git a/Engine/Plugins/VirtualProduction/ActorModifierCore/Source/ActorModifierCore/Private/Modifiers/ActorModifierCoreSharedActor.cpp b/Engine/Plugins/VirtualProduction/ActorModifierCore/Source/ActorModifierCore/Private/Modifiers/ActorModifierCoreSharedActor.cpp
--- a/Engine/Plugins/VirtualProduction/ActorModifierCore/Source/ActorModifierCore/Private/Modifiers/ActorModifierCoreSharedActor.cpp
+++ b/Engine/Plugins/VirtualProduction/ActorModifierCore/Source/ActorModifierCore/Private/Modifiers/ActorModifierCoreSharedActor.cpp
@@ -29,7 +29,7 @@ void AActorModifierCoreSharedActor::PostLoad()
 			continue;
 		}
 
-		UE_LOG(LogActorModifierCoreSharedActor, Log, TEXT("%s : Modifier Shared object of class %s loaded"), *GetActorNameOrLabel(), *It->Key->GetName());
+		UE_LOG(LogActorModifierCoreSharedActor, Verbose, TEXT("%s : Modifier Shared object of class %s loaded"), *GetActorNameOrLabel(), *It->Key->GetName());
 	}
 }
 
@@ -48,7 +48,7 @@ UActorModifierCoreSharedObject* AActorModifierCoreSharedActor::FindOrAddShared(T
 	UActorModifierCoreSharedObject* NewSharedObject = NewObject<UActorModifierCoreSharedObject>(this, InSharedClass);
 	SharedObjects.Add(InSharedClass, NewSharedObject);
 
-	UE_LOG(LogActorModifierCoreSharedActor, Log, TEXT("%s : Modifier Shared object of class %s created"), *GetActorNameOrLabel(), *InSharedClass->GetName());
+	UE_LOG(LogActorModifierCoreSharedActor, Verbose, TEXT("%s : Modifier Shared object of class %s created"), *GetActorNameOrLabel(), *InSharedClass->GetName());
 
 	return NewSharedObject;
 }
diff --git a/Engine/Plugins/VirtualProduction/Avalanche/Source/AvalancheMask/Private/Handling/AvaObjectHandleSubsystem.cpp b/Engine/Plugins/VirtualProduction/Avalanche/Source/AvalancheMask/Private/Handling/AvaObjectHandleSubsystem.cpp
--- a/Engine/Plugins/VirtualProduction/Avalanche/Source/AvalancheMask/Private/Handling/AvaObjectHandleSubsystem.cpp
+++ b/Engine/Plugins/VirtualProduction/Avalanche/Source/AvalancheMask/Private/Handling/AvaObjectHandleSubsystem.cpp
@@ -51,7 +51,7 @@ TSharedPtr<IAvaObjectHandle> UAvaObjectHandleSubsystem::MakeHandleDirect(const F
 		}
 	}
 
-	UE_LOG(LogAvaMask, Display, TEXT("No ObjectHandle found for '%s'"), *InInstance.ToString());
+	UE_LOG(LogAvaMask, Verbose, TEXT("No ObjectHandle found for '%s'"), *InInstance.ToString());
 	return nullptr;
 }
 
diff --git a/Engine/Plugins/VirtualProduction/Avalanche/Source/AvalancheMask/Private/Handling/IAvaMaskMaterialCollectionHandle.h b/Engine/Plugins/VirtualProduction/Avalanche/Source/AvalancheMask/Private/Handling/IAvaMaskMaterialCollectionHandle.h
--- a/Engine/Plugins/VirtualProduction/Avalanche/Source/AvalancheMask/Private/Handling/IAvaMaskMaterialCollectionHandle.h
+++ b/Engine/Plugins/VirtualProduction/Avalanche/Source/AvalancheMask/Private/Handling/IAvaMaskMaterialCollectionHandle.h
@@ -250,7 +250,7 @@ bool TAvaMaskMaterialCollectionHandle<HandleDataType>::ApplyOriginalState(const
 					});
 			}
 
-			UE_LOG(LogAvaMask, Display, TEXT("MaterialData not found for slot %u"), InSlotIdx);
+			UE_LOG(LogAvaMask, Verbose, TEXT("MaterialData not found for slot %u"), InSlotIdx);
 
 			return false;
 		});
@@ -294,7 +294,7 @@ bool TAvaMaskMaterialCollectionHandle<HandleDataType>::ApplyModifiedState(
 					});
 			}
 
-			UE_LOG(LogAvaMask, Display, TEXT("MaterialData not found for slot %u"), InSlotIdx);
+			UE_LOG(LogAvaMask, Verbose, TEXT("MaterialData not found for slot %u"), InSlotIdx);
 
 			return false;
 		});
diff --git a/Engine/Plugins/VirtualProduction/Avalanche/Source/AvalancheMask/Private/Mask2D/AvaMask2DBaseModifier.cpp b/Engine/Plugins/VirtualProduction/Avalanche/Source/AvalancheMask/Private/Mask2D/AvaMask2DBaseModifier.cpp
--- a/Engine/Plugins/VirtualProduction/Avalanche/Source/AvalancheMask/Private/Mask2D/AvaMask2DBaseModifier.cpp
+++ b/Engine/Plugins/VirtualProduction/Avalanche/Source/AvalancheMask/Private/Mask2D/AvaMask2DBaseModifier.cpp
@@ -584,7 +584,7 @@ UActorComponent* UAvaMask2DBaseModifier::FindOrAddMaskComponent(TSubclassOf<UAct
 	{
 		if (InActor->FindComponentByInterface(UGeometryMaskWriteInterface::StaticClass()))
 		{
-			UE_LOG(LogAvaMask, Display, TEXT("Attempting to add a Mask Read component to an Actor that has a Mask Write component: '%s'"), *InActor->GetName());
+			UE_LOG(LogAvaMask, Verbose, TEXT("Attempting to add a Mask Read component to an Actor that has a Mask Write component: '%s'"), *InActor->GetName());
 			return nullptr;
 		}
 	}
diff --git a/Engine/Plugins/VirtualProduction/Text3D/Source/Text3D/Private/Renderers/StaticMeshes/Text3DStaticMeshesRenderer.cpp b/Engine/Plugins/VirtualProduction/Text3D/Source/Text3D/Private/Renderers/StaticMeshes/Text3DStaticMeshesRenderer.cpp
--- a/Engine/Plugins/VirtualProduction/Text3D/Source/Text3D/Private/Renderers/StaticMeshes/Text3DStaticMeshesRenderer.cpp
+++ b/Engine/Plugins/VirtualProduction/Text3D/Source/Text3D/Private/Renderers/StaticMeshes/Text3DStaticMeshesRenderer.cpp
@@ -319,8 +319,6 @@ void UText3DStaticMeshesRenderer::ConvertToStaticMesh()
 	FScopedSlowTask SlowTask(0.0f, LOCTEXT("ConvertToStaticMesh", "Converting Text3D to static mesh"));
 	SlowTask.MakeDialog();
 
-	UE_LOG(LogText3D, Log, TEXT("%s : Request ConvertToStaticMesh..."), *Owner->GetActorNameOrLabel())
-
 	if (UE::Text3D::Utilities::Conversion::ConvertToStaticMesh(Text3DComponent))
 	{
 		UE_LOG(LogText3D, Log, TEXT("%s : ConvertToStaticMesh Completed"), *Owner->GetActorNameOrLabel())
diff --git a/Engine/Plugins/VirtualProduction/Text3D/Source/Text3D/Private/Renderers/Text3DRendererBase.cpp b/Engine/Plugins/VirtualProduction/Text3D/Source/Text3D/Private/Renderers/Text3DRendererBase.cpp
--- a/Engine/Plugins/VirtualProduction/Text3D/Source/Text3D/Private/Renderers/Text3DRendererBase.cpp
+++ b/Engine/Plugins/VirtualProduction/Text3D/Source/Text3D/Private/Renderers/Text3DRendererBase.cpp
@@ -89,7 +89,7 @@ void UText3DRendererBase::Create()
 	bInitialized = true;
 
 	const AActor* Owner = Text3DComponent->GetOwner();
-	UE_LOG(LogText3D, Log, TEXT("%s : Text3DRenderer %s Created"), !!Owner ? *Owner->GetActorNameOrLabel() : TEXT("Invalid owner"), *GetFriendlyName().ToString())
+	UE_LOG(LogText3D, Verbose, TEXT("%s : Text3DRenderer %s Created"), !!Owner ? *Owner->GetActorNameOrLabel() : TEXT("Invalid owner"), *GetFriendlyName().ToString())
 }
 
 void UText3DRendererBase::Update(EText3DRendererFlags InFlags)
@@ -238,7 +238,7 @@ void UText3DRendererBase::Destroy()
 	bInitialized = false;
 
 	const AActor* Owner = Text3DComponent->GetOwner();
-	UE_LOG(LogText3D, Log, TEXT("%s : Text3DRenderer %s Destroyed"), !!Owner ? *Owner->GetActorNameOrLabel() : TEXT("Invalid owner"), *GetFriendlyName().ToString())
+	UE_LOG(LogText3D, Verbose, TEXT("%s : Text3DRenderer %s Destroyed"), !!Owner ? *Owner->GetActorNameOrLabel() : TEXT("Invalid owner"), *GetFriendlyName().ToString())
 }
 
 FBox UText3DRendererBase::GetBounds() const
diff --git a/Engine/Plugins/VirtualProduction/Text3D/Source/Text3D/Private/Subsystems/Text3DEngineSubsystem.cpp b/Engine/Plugins/VirtualProduction/Text3D/Source/Text3D/Private/Subsystems/Text3DEngineSubsystem.cpp
--- a/Engine/Plugins/VirtualProduction/Text3D/Source/Text3D/Private/Subsystems/Text3DEngineSubsystem.cpp
+++ b/Engine/Plugins/VirtualProduction/Text3D/Source/Text3D/Private/Subsystems/Text3DEngineSubsystem.cpp
@@ -114,7 +114,7 @@ bool UText3DEngineSubsystem::CleanupCache(float)
 
 	if (FontFaceRemoved > 0 || GlyphMeshRemoved > 0)
 	{
-		UE_LOG(LogText3D, Log, TEXT("Text3D Engine Subsystem cache cleanup completed : %i unused glyph(s) removed, %i remaining glyph(s), %i unused font face(s) removed, %i remaining font face(s)"), GlyphMeshRemoved, GlyphMeshCount, FontFaceRemoved, CachedFontFaces.Num());
+		UE_LOG(LogText3D, Verbose, TEXT("Text3D Engine Subsystem cache cleanup completed : %i unused glyph(s) removed, %i remaining glyph(s), %i unused font face(s) removed, %i remaining font face(s)"), GlyphMeshRemoved, GlyphMeshCount, FontFaceRemoved, CachedFontFaces.Num());
 	}
 
 	return true;
