Motion Design Actor Modifiers: added CanMarkDirty to modifiers to check ahead of time if the modifier can't be dirtied before doing additional operations (e.g. in RenderStateDirty)

#Synced-CL 48854462

diff --git a/Engine/Plugins/VirtualProduction/ActorModifierCore/Source/ActorModifierCore/Private/Modifiers/ActorModifierCoreBase.cpp b/Engine/Plugins/VirtualProduction/ActorModifierCore/Source/ActorModifierCore/Private/Modifiers/ActorModifierCoreBase.cpp
--- a/Engine/Plugins/VirtualProduction/ActorModifierCore/Source/ActorModifierCore/Private/Modifiers/ActorModifierCoreBase.cpp
+++ b/Engine/Plugins/VirtualProduction/ActorModifierCore/Source/ActorModifierCore/Private/Modifiers/ActorModifierCoreBase.cpp
@@ -318,28 +318,35 @@ const UActorModifierCoreBase* UActorModifierCoreBase::GetNextNameModifier(const
 	return NextNameModifier;
 }
 
-void UActorModifierCoreBase::MarkModifierDirty(bool bExecute)
+bool UActorModifierCoreBase::CanMarkModifierDirty() const
 {
 	const UActorModifierCoreStack* Stack = GetRootModifierStack();
 
 	// Do not mark dirty if stack is in execution
 	if (!Stack || !Stack->IsModifierIdle() || !Stack->IsModifierInitialized() || !IsModifierInitialized())
 	{
-		return;
+		return false;
 	}
 
 	// When modifier is disabled but applied, we need to mark it dirty for it to be restored
 	if (!IsModifierEnabled() && !IsModifierApplied())
 	{
-		return;
+		return false;
 	}
 
-	if (!bModifierDirty || bExecute)
+	// Can't mark modifier dirty. It's already dirtied.
+	return !bModifierDirty;
+}
+
+bool UActorModifierCoreBase::MarkModifierDirty(bool bExecute)
+{
+	if (CanMarkModifierDirty() || bExecute)
 	{
 		bModifierDirty = true;
-
 		OnModifierDirty(this, bExecute);
+		return true;
 	}
+	return false;
 }
 
 bool UActorModifierCoreBase::IsModifierEnabled() const
diff --git a/Engine/Plugins/VirtualProduction/ActorModifierCore/Source/ActorModifierCore/Public/Modifiers/ActorModifierCoreBase.h b/Engine/Plugins/VirtualProduction/ActorModifierCore/Source/ActorModifierCore/Public/Modifiers/ActorModifierCoreBase.h
--- a/Engine/Plugins/VirtualProduction/ActorModifierCore/Source/ActorModifierCore/Public/Modifiers/ActorModifierCoreBase.h
+++ b/Engine/Plugins/VirtualProduction/ActorModifierCore/Source/ActorModifierCore/Public/Modifiers/ActorModifierCoreBase.h
@@ -87,8 +87,14 @@ public:
 	/** Returns the first modifier after this one that matches the name, will recurse upwards, could return a stack */
 	const UActorModifierCoreBase* GetNextNameModifier(const FName& InModifierName) const;
 
-	/** Set this modifier as dirty to be able to execute only dirty modifiers */
-	ACTORMODIFIERCORE_API void MarkModifierDirty(bool bExecute = true);
+	/** Returns whether a modifier can be marked dirty */
+	ACTORMODIFIERCORE_API bool CanMarkModifierDirty() const;
+
+	/** 
+	 * Set this modifier as dirty to be able to execute only dirty modifiers 
+	 * returns true if the modifier could be marked dirty
+	 */
+	ACTORMODIFIERCORE_API bool MarkModifierDirty(bool bExecute = true);
 
 	/** Whether this modifier is set as dirty */
 	bool IsModifierDirty() const
