Fix modifier map removing entries while iterating the map instead of using the Iterator.RemoveCurrent option.

#Synced-CL 49962738

diff --git a/Engine/Plugins/VirtualProduction/ActorModifierCore/Source/ActorModifierCore/Private/Subsystems/ActorModifierCoreSubsystem.cpp b/Engine/Plugins/VirtualProduction/ActorModifierCore/Source/ActorModifierCore/Private/Subsystems/ActorModifierCoreSubsystem.cpp
--- a/Engine/Plugins/VirtualProduction/ActorModifierCore/Source/ActorModifierCore/Private/Subsystems/ActorModifierCoreSubsystem.cpp
+++ b/Engine/Plugins/VirtualProduction/ActorModifierCore/Source/ActorModifierCore/Private/Subsystems/ActorModifierCoreSubsystem.cpp
@@ -178,7 +178,9 @@ bool UActorModifierCoreSubsystem::UnregisterModifierClass(const UClass* InModifi
 	{
 		if (It->Value->GetClass() == InModifierClass)
 		{
-			bUnregistered |= UnregisterModifierClass(It->Key);
+			UActorModifierCoreSubsystem::OnModifierClassUnregisteredDelegate.Broadcast(*It.Value());
+			It.RemoveCurrent();
+			bUnregistered = true;
 		}
 	}
 
