Motion Design: made render state dirty reason optional + added some fixes to the text3d update causing new characters to not be updated by the translucent priority modifier.
The render state dirty reason will be defaulted to disabled as this is only there for performance reasons.

#Synced-CL 52400047

diff --git a/Engine/Plugins/VirtualProduction/ActorModifier/Source/ActorModifier/Private/Extensions/ActorModifierRenderStateUpdateExtension.cpp b/Engine/Plugins/VirtualProduction/ActorModifier/Source/ActorModifier/Private/Extensions/ActorModifierRenderStateUpdateExtension.cpp
--- a/Engine/Plugins/VirtualProduction/ActorModifier/Source/ActorModifier/Private/Extensions/ActorModifierRenderStateUpdateExtension.cpp
+++ b/Engine/Plugins/VirtualProduction/ActorModifier/Source/ActorModifier/Private/Extensions/ActorModifierRenderStateUpdateExtension.cpp
@@ -97,7 +97,7 @@ void FActorModifierRenderStateUpdateExtension::OnExtensionDisabled(EActorModifie
 
 void FActorModifierRenderStateUpdateExtension::OnRenderStateDirty(UActorComponent& InComponent)
 {
-	if (!UE::ActorModifierCore::IsRenderStateDirtyReasonPresent(RelevantReasons))
+	if (!UE::ActorModifierCore::IsRenderStateDirtyRelevant(RelevantReasons))
 	{
 		return;
 	}
@@ -127,7 +127,7 @@ void FActorModifierRenderStateUpdateExtension::OnRenderStateDirty(UActorComponen
 	}
 
 	// Skip if render state is not relevant or not tracking the actor
-	if (!UE::ActorModifierCore::IsRenderStateDirtyReasonPresent(UE::ActorModifierCore::ERenderStateDirtyReason::Visibility)
+	if (!UE::ActorModifierCore::IsRenderStateDirtyRelevant(UE::ActorModifierCore::ERenderStateDirtyReason::Visibility)
 		|| !TrackedActorsVisibility.Contains(ActorDirty))
 	{
 		return;
diff --git a/Engine/Plugins/VirtualProduction/ActorModifier/Source/ActorModifier/Private/Extensions/ActorModifierSceneTreeUpdateExtension.cpp b/Engine/Plugins/VirtualProduction/ActorModifier/Source/ActorModifier/Private/Extensions/ActorModifierSceneTreeUpdateExtension.cpp
--- a/Engine/Plugins/VirtualProduction/ActorModifier/Source/ActorModifier/Private/Extensions/ActorModifierSceneTreeUpdateExtension.cpp
+++ b/Engine/Plugins/VirtualProduction/ActorModifier/Source/ActorModifier/Private/Extensions/ActorModifierSceneTreeUpdateExtension.cpp
@@ -156,7 +156,7 @@ void FActorModifierSceneTreeUpdateExtension::OnRenderStateDirty(UActorComponent&
 	constexpr UE::ActorModifierCore::ERenderStateDirtyReason RelevantReasons = UE::ActorModifierCore::ERenderStateDirtyReason::Visibility 
 		| UE::ActorModifierCore::ERenderStateDirtyReason::SceneTree;
 
-	if (!UE::ActorModifierCore::IsRenderStateDirtyReasonPresent(RelevantReasons))
+	if (!UE::ActorModifierCore::IsRenderStateDirtyRelevant(RelevantReasons))
 	{
 		return;
 	}
diff --git a/Engine/Plugins/VirtualProduction/ActorModifierCore/Source/ActorModifierCore/Private/Modifiers/ActorModifierRenderStateDirtyEvent.cpp b/Engine/Plugins/VirtualProduction/ActorModifierCore/Source/ActorModifierCore/Private/Modifiers/ActorModifierRenderStateDirtyEvent.cpp
--- a/Engine/Plugins/VirtualProduction/ActorModifierCore/Source/ActorModifierCore/Private/Modifiers/ActorModifierRenderStateDirtyEvent.cpp
+++ b/Engine/Plugins/VirtualProduction/ActorModifierCore/Source/ActorModifierCore/Private/Modifiers/ActorModifierRenderStateDirtyEvent.cpp
@@ -1,13 +1,26 @@
 // Copyright Epic Games, Inc. All Rights Reserved.
 
 #include "Modifiers/ActorModifierRenderStateDirtyEvent.h"
+#include "HAL/IConsoleManager.h"
 
 namespace UE::ActorModifierCore
 {
 namespace Private
 {
-	// Defaults to unknown as most MarkRenderStateDirty calls will not have a reason scope to it
-	static ERenderStateDirtyReason RenderStateDirtyReason = ERenderStateDirtyReason::Unknown;
+
+/**
+ * Determines whether render state reason should be enabled.
+ * If disabled, IsRenderStateDirtyRelevant() will always return true (no specific reason to otherwise make it irrelevant)
+ */
+static bool bRenderStateDirtyReasonEnabled = false;
+static FAutoConsoleVariableRef CVarRenderStateDirtyReasonEnabled(
+	TEXT("ActorModifierCore.EnableRenderStateDirtyReason"), 
+	bRenderStateDirtyReasonEnabled,
+	TEXT("If true some render state dirty events will have a reason to them and modifiers will skip work if the reason is deemed irrelevant by the modifier."),
+	ECVF_Default);
+
+// Defaults to unknown as most MarkRenderStateDirty calls will not have a reason scope to it
+static ERenderStateDirtyReason RenderStateDirtyReason = ERenderStateDirtyReason::Unknown;
 
 } // UE::ActorModifierCore::Private
 
@@ -16,16 +29,17 @@ ERenderStateDirtyReason GetRenderStateDirtyReason()
 	return Private::RenderStateDirtyReason;
 }
 
-bool IsRenderStateDirtyReasonPresent(ERenderStateDirtyReason InReason)
+bool IsRenderStateDirtyRelevant(ERenderStateDirtyReason InReason)
 {
-	return Private::RenderStateDirtyReason == ERenderStateDirtyReason::Unknown
+	return !Private::bRenderStateDirtyReasonEnabled
+		|| Private::RenderStateDirtyReason == ERenderStateDirtyReason::Unknown
 		|| EnumHasAnyFlags(Private::RenderStateDirtyReason, InReason); 
 }
 
 FRenderStateDirtyReasonScope::FRenderStateDirtyReasonScope(ERenderStateDirtyReason InReason)
 	: PreviousReason(Private::RenderStateDirtyReason)
 {
-	Private::RenderStateDirtyReason |= InReason;
+	Private::RenderStateDirtyReason = InReason;
 }
 
 FRenderStateDirtyReasonScope::~FRenderStateDirtyReasonScope()
diff --git a/Engine/Plugins/VirtualProduction/ActorModifierCore/Source/ActorModifierCore/Public/Modifiers/ActorModifierRenderStateDirtyEvent.h b/Engine/Plugins/VirtualProduction/ActorModifierCore/Source/ActorModifierCore/Public/Modifiers/ActorModifierRenderStateDirtyEvent.h
--- a/Engine/Plugins/VirtualProduction/ActorModifierCore/Source/ActorModifierCore/Public/Modifiers/ActorModifierRenderStateDirtyEvent.h
+++ b/Engine/Plugins/VirtualProduction/ActorModifierCore/Source/ActorModifierCore/Public/Modifiers/ActorModifierRenderStateDirtyEvent.h
@@ -37,7 +37,7 @@ UE_API ERenderStateDirtyReason GetRenderStateDirtyReason();
  * Determines whether the given reason is part of the current render state dirty reason
  * When the current reason is 'unknown' this returns true.
  */
-UE_API bool IsRenderStateDirtyReasonPresent(ERenderStateDirtyReason InReason);
+UE_API bool IsRenderStateDirtyRelevant(ERenderStateDirtyReason InReason);
 
 /**
  * Tracks the reason a render state dirty happened.
diff --git a/Engine/Plugins/VirtualProduction/Avalanche/Source/AvalancheText/Private/AvaTextModule.cpp b/Engine/Plugins/VirtualProduction/Avalanche/Source/AvalancheText/Private/AvaTextModule.cpp
--- a/Engine/Plugins/VirtualProduction/Avalanche/Source/AvalancheText/Private/AvaTextModule.cpp
+++ b/Engine/Plugins/VirtualProduction/Avalanche/Source/AvalancheText/Private/AvaTextModule.cpp
@@ -38,22 +38,26 @@ void FAvaTextModule::RegisterTextRendererDelegates()
 {
 	static TArray<TUniquePtr<UE::ActorModifierCore::FRenderStateDirtyReasonScope>, TInlineAllocator<8>> Scope;
 
-	TextPreRendererUpdateHandle = UE::Text3D::OnText3DPreRendererUpdate().AddLambda(
-		[](const UE::Text3D::FText3DBuildContext& InBuildContext)
+	static auto GetReason = [](const UE::Text3D::FText3DBuildContext& InBuildContext)->UE::ActorModifierCore::ERenderStateDirtyReason
 		{
-			if (InBuildContext.CurrentFlag == EText3DRendererFlags::Material)
+			// If only updating material, use the reason scope. Otherwise, keep it unknown
+			if (InBuildContext.UpdateFlags == EText3DRendererFlags::Material)
 			{
-				Scope.Push(MakeUnique<UE::ActorModifierCore::FRenderStateDirtyReasonScope>(UE::ActorModifierCore::ERenderStateDirtyReason::Material));
+				return UE::ActorModifierCore::ERenderStateDirtyReason::Material;
 			}
+			return UE::ActorModifierCore::ERenderStateDirtyReason::Unknown;
+		};
+
+	TextPreRendererUpdateHandle = UE::Text3D::OnText3DPreRendererUpdate().AddLambda(
+		[](const UE::Text3D::FText3DBuildContext& InBuildContext)
+		{
+			Scope.Push(MakeUnique<UE::ActorModifierCore::FRenderStateDirtyReasonScope>(GetReason(InBuildContext)));
 		});
 
 	TextPostRendererUpdateHandle = UE::Text3D::OnText3DPostRendererUpdate().AddLambda(
 		[](const UE::Text3D::FText3DBuildContext& InBuildContext)
 		{
-			if (InBuildContext.CurrentFlag == EText3DRendererFlags::Material)
-			{
-				Scope.Pop();
-			}
+			Scope.Pop();
 		});
 }
 
diff --git a/Engine/Plugins/VirtualProduction/Text3D/Source/Text3D/Private/Renderers/StaticMeshes/Text3DStaticMeshesRenderer.cpp b/Engine/Plugins/VirtualProduction/Text3D/Source/Text3D/Private/Renderers/StaticMeshes/Text3DStaticMeshesRenderer.cpp
--- a/Engine/Plugins/VirtualProduction/Text3D/Source/Text3D/Private/Renderers/StaticMeshes/Text3DStaticMeshesRenderer.cpp
+++ b/Engine/Plugins/VirtualProduction/Text3D/Source/Text3D/Private/Renderers/StaticMeshes/Text3DStaticMeshesRenderer.cpp
@@ -373,6 +373,7 @@ void UText3DStaticMeshesRenderer::AllocateComponents(int32 InCount)
 					MeshComponent->UnregisterComponent();
 				}
 				Actor->RemoveOwnedComponent(MeshComponent);
+				Actor->RemoveInstanceComponent(MeshComponent);
 
 				if (MeshComponent->GetOwner() == Actor)
 				{
@@ -384,6 +385,7 @@ void UText3DStaticMeshesRenderer::AllocateComponents(int32 InCount)
 	else if (RemainingCharacterCount < 0)
 	{
 		const int32 RemainingCharacterCountPositive = FMath::Abs(RemainingCharacterCount);
+		const FName CharacterMeshName = TEXT("CharacterMesh");
 		for (int32 CharacterIndex = 0; CharacterIndex < RemainingCharacterCountPositive; ++CharacterIndex)
 		{
 			UStaticMeshComponent* MeshComponent = nullptr;
@@ -394,7 +396,8 @@ void UText3DStaticMeshesRenderer::AllocateComponents(int32 InCount)
 
 			if (!MeshComponent)
 			{
-				MeshComponent = NewObject<UStaticMeshComponent>(this);
+				const FName MeshComponentName = MakeUniqueObjectName(this, UStaticMeshComponent::StaticClass(), CharacterMeshName);
+				MeshComponent = NewObject<UStaticMeshComponent>(this, MeshComponentName);
 				MeshComponent->CreationMethod = EComponentCreationMethod::Instance;
 			}
 
diff --git a/Engine/Plugins/VirtualProduction/Text3D/Source/Text3D/Private/Renderers/Text3DRendererBase.cpp b/Engine/Plugins/VirtualProduction/Text3D/Source/Text3D/Private/Renderers/Text3DRendererBase.cpp
--- a/Engine/Plugins/VirtualProduction/Text3D/Source/Text3D/Private/Renderers/Text3DRendererBase.cpp
+++ b/Engine/Plugins/VirtualProduction/Text3D/Source/Text3D/Private/Renderers/Text3DRendererBase.cpp
@@ -97,10 +97,7 @@ void UText3DRendererBase::Update(const UE::Text3D::Renderer::FUpdateParameters&
 {
 	if (bInitialized)
 	{
-		const UE::Text3D::FText3DBuildContext BuildContext
-			{
-				.CurrentFlag = InParameters.CurrentFlag
-			};
+		const UE::Text3D::FText3DBuildContext BuildContext(InParameters);
 
 		UE::Text3D::BroadcastPreRendererUpdate(BuildContext);
 		OnUpdate(InParameters);
diff --git a/Engine/Plugins/VirtualProduction/Text3D/Source/Text3D/Private/Text3DDelegates.cpp b/Engine/Plugins/VirtualProduction/Text3D/Source/Text3D/Private/Text3DDelegates.cpp
--- a/Engine/Plugins/VirtualProduction/Text3D/Source/Text3D/Private/Text3DDelegates.cpp
+++ b/Engine/Plugins/VirtualProduction/Text3D/Source/Text3D/Private/Text3DDelegates.cpp
@@ -21,6 +21,11 @@ TMulticastDelegate<void(const FText3DBuildContext&)>& GetPostRendererUpdateDeleg
 
 } // UE::Text3D::Private
 
+FText3DBuildContext::FText3DBuildContext(const Renderer::FUpdateParameters& InUpdateParameters)
+	: UpdateFlags(InUpdateParameters.UpdateFlags)
+{
+}
+
 void BroadcastPreRendererUpdate(const FText3DBuildContext& InBuildContext)
 {
 	Private::GetPreRendererUpdateDelegate().Broadcast(InBuildContext);
diff --git a/Engine/Plugins/VirtualProduction/Text3D/Source/Text3D/Public/Text3DDelegates.h b/Engine/Plugins/VirtualProduction/Text3D/Source/Text3D/Public/Text3DDelegates.h
--- a/Engine/Plugins/VirtualProduction/Text3D/Source/Text3D/Public/Text3DDelegates.h
+++ b/Engine/Plugins/VirtualProduction/Text3D/Source/Text3D/Public/Text3DDelegates.h
@@ -12,8 +12,10 @@ namespace UE::Text3D
 
 struct UE_INTERNAL FText3DBuildContext
 {
-	/** The current flag to update in the renderer */
-	EText3DRendererFlags CurrentFlag = EText3DRendererFlags::None;
+	explicit FText3DBuildContext(const Renderer::FUpdateParameters& InUpdateParameters);
+
+	/** All dirty flags for this update */
+	EText3DRendererFlags UpdateFlags = EText3DRendererFlags::None;
 };
 
 /** Called before the Text3D renderer update */
