Actor Modifiers: added experimental freeze modifier feature.

#Synced-CL 52400047

diff --git a/Engine/Plugins/VirtualProduction/ActorModifierCore/Source/ActorModifierCore/Private/Modifiers/ActorModifierCoreBase.cpp b/Engine/Plugins/VirtualProduction/ActorModifierCore/Source/ActorModifierCore/Private/Modifiers/ActorModifierCoreBase.cpp
--- a/Engine/Plugins/VirtualProduction/ActorModifierCore/Source/ActorModifierCore/Private/Modifiers/ActorModifierCoreBase.cpp
+++ b/Engine/Plugins/VirtualProduction/ActorModifierCore/Source/ActorModifierCore/Private/Modifiers/ActorModifierCoreBase.cpp
@@ -356,7 +356,7 @@ bool UActorModifierCoreBase::IsModifierEnabled() const
 		return false;
 	}
 	const UActorModifierCoreStack* Stack = GetModifierStack();
-	return !Stack || Stack->IsModifierEnabled();
+	return !Stack || (Stack->IsModifierEnabled() && !Stack->IsModifierStackFrozen());
 }
 
 void UActorModifierCoreBase::ProcessLockFunction(TFunctionRef<void()> InFunction)
diff --git a/Engine/Plugins/VirtualProduction/ActorModifierCore/Source/ActorModifierCore/Private/Modifiers/ActorModifierCoreDefs.cpp b/Engine/Plugins/VirtualProduction/ActorModifierCore/Source/ActorModifierCore/Private/Modifiers/ActorModifierCoreDefs.cpp
--- a/Engine/Plugins/VirtualProduction/ActorModifierCore/Source/ActorModifierCore/Private/Modifiers/ActorModifierCoreDefs.cpp
+++ b/Engine/Plugins/VirtualProduction/ActorModifierCore/Source/ActorModifierCore/Private/Modifiers/ActorModifierCoreDefs.cpp
@@ -527,7 +527,8 @@ void FActorModifierCoreExecutionTask::ExecuteCurrentModifier()
 
 	if (bValidActor && bModifierReady)
 	{
-		if (Modifier->bModifierEnabled || Modifier->IsModifierStack())
+		const bool bModifierStackFrozen = ModifierStack && ModifierStack->IsModifierStackFrozen();
+		if (!bModifierStackFrozen && (Modifier->bModifierEnabled || Modifier->IsModifierStack()))
 		{
 			// lock current execution state
 			Modifier->BeginModifierExecution();
diff --git a/Engine/Plugins/VirtualProduction/ActorModifierCore/Source/ActorModifierCore/Private/Modifiers/ActorModifierCoreStack.cpp b/Engine/Plugins/VirtualProduction/ActorModifierCore/Source/ActorModifierCore/Private/Modifiers/ActorModifierCoreStack.cpp
--- a/Engine/Plugins/VirtualProduction/ActorModifierCore/Source/ActorModifierCore/Private/Modifiers/ActorModifierCoreStack.cpp
+++ b/Engine/Plugins/VirtualProduction/ActorModifierCore/Source/ActorModifierCore/Private/Modifiers/ActorModifierCoreStack.cpp
@@ -40,6 +40,22 @@ UActorModifierCoreStack* UActorModifierCoreStack::Create(UActorModifierCoreCompo
 	return nullptr;
 }
 
+bool UActorModifierCoreStack::IsModifierStackFrozen() const
+{
+	return bModifierStackFrozen;
+}
+
+void UActorModifierCoreStack::SetModifierStackFrozen(bool bInFrozen)
+{
+	if (bModifierStackFrozen == bInFrozen)
+	{
+		return;
+	}
+
+	bModifierStackFrozen = bInFrozen;
+	OnModifierStackFrozenChanged();
+}
+
 bool UActorModifierCoreStack::IsModifierReady() const
 {
 	for (const TObjectPtr<UActorModifierCoreBase>& Modifier : Modifiers)
@@ -106,6 +122,18 @@ void UActorModifierCoreStack::PostLoad()
 	Super::PostLoad();
 }
 
+#if WITH_EDITOR
+void UActorModifierCoreStack::PostEditChangeProperty(FPropertyChangedEvent& InPropertyChangedEvent)
+{
+	Super::PostEditChangeProperty(InPropertyChangedEvent);
+
+	if (InPropertyChangedEvent.GetMemberPropertyName() == GET_MEMBER_NAME_CHECKED(UActorModifierCoreStack, bModifierStackFrozen))
+	{
+		OnModifierStackFrozenChanged();
+	}
+}
+#endif
+
 void UActorModifierCoreStack::OnModifierAdded(EActorModifierCoreEnableReason InReason)
 {
 	Super::OnModifierAdded(InReason);
@@ -488,6 +516,15 @@ bool UActorModifierCoreStack::ContainsModifierAfter(const UActorModifierCoreBase
 
 UActorModifierCoreBase* UActorModifierCoreStack::CloneModifier(FActorModifierCoreStackCloneOp& InCloneOp)
 {
+	if (IsModifierStackFrozen())
+	{
+		if (InCloneOp.FailReason)
+		{
+			*InCloneOp.FailReason = LOCTEXT("CloneFailed_FrozenStack", "Modifier Stack is Frozen");
+		}
+		return nullptr;
+	}
+
 	UActorModifierCoreBase* NewInsertedModifier = nullptr;
 
 	if (!IsValid(InCloneOp.CloneModifier) || InCloneOp.CloneModifier->GetModifiedActor() == GetModifiedActor())
@@ -683,6 +720,15 @@ UActorModifierCoreBase* UActorModifierCoreStack::CloneModifier(FActorModifierCor
 
 UActorModifierCoreBase* UActorModifierCoreStack::InsertModifier(FActorModifierCoreStackInsertOp& InInsertOp)
 {
+	if (IsModifierStackFrozen())
+	{
+		if (InInsertOp.FailReason)
+		{
+			*InInsertOp.FailReason = LOCTEXT("InsertFailed_FrozenStack", "Modifier Stack is Frozen");
+		}
+		return nullptr;
+	}
+
 	UActorModifierCoreBase* NewInsertedModifier = nullptr;
 
 	if (InInsertOp.NewModifierName.IsNone())
@@ -866,6 +912,15 @@ UActorModifierCoreBase* UActorModifierCoreStack::InsertModifier(FActorModifierCo
 
 bool UActorModifierCoreStack::MoveModifier(FActorModifierCoreStackMoveOp& InMoveOp)
 {
+	if (IsModifierStackFrozen())
+	{
+		if (InMoveOp.FailReason)
+		{
+			*InMoveOp.FailReason = LOCTEXT("MoveFailed_FrozenStack", "Modifier Stack is Frozen");
+		}
+		return false;
+	}
+
 	if (!IsValid(InMoveOp.MoveModifier))
 	{
 		return false;
@@ -1131,6 +1186,15 @@ bool UActorModifierCoreStack::MoveModifier(FActorModifierCoreStackMoveOp& InMove
 
 bool UActorModifierCoreStack::RemoveModifier(FActorModifierCoreStackRemoveOp& InRemoveOp)
 {
+	if (IsModifierStackFrozen())
+	{
+		if (InRemoveOp.FailReason)
+		{
+			*InRemoveOp.FailReason = LOCTEXT("RemoveFailed_FrozenStack", "Modifier Stack is Frozen");
+		}
+		return false;
+	}
+
 	if (!IsValid(InRemoveOp.RemoveModifier))
 	{
 		if (InRemoveOp.FailReason)
@@ -1258,7 +1322,15 @@ bool UActorModifierCoreStack::RemoveModifier(FActorModifierCoreStackRemoveOp& In
 
 			if (This && RemoveModifier)
 			{
-				RemoveModifier->UninitializeModifier(EActorModifierCoreDisableReason::User);
+				if (This->IsModifierStackFrozen())
+				{
+					// Broadcast the removal directly without going through uninit process
+					UActorModifierCoreStack::OnModifierRemovedDelegate.Broadcast(RemoveModifier, EActorModifierCoreDisableReason::User);
+				}
+				else
+				{
+					RemoveModifier->UninitializeModifier(EActorModifierCoreDisableReason::User);
+				}
 
 				This->ExecutionTask.RemoveModifier(*RemoveModifier);
 
@@ -1360,12 +1432,20 @@ bool UActorModifierCoreStack::GetRequiredModifiers(UActorModifierCoreBase* InMod
 bool UActorModifierCoreStack::RemoveAllModifiers()
 {
 	LockModifierExecution();
+
+	// Temporarily unfreeze modifier stack to allow removing all modifiers
+	// ProcessIdleFunction checks for IsModifierStackFrozen when this goes out of scope. And it's ok.
+	// This is only here to overcome the first check blocking removal when stack is frozen
+	TGuardValue<bool> StackUnfrozenScope(bModifierStackFrozen, false);
+
 	for (int32 ModifierIdx = Modifiers.Num() - 1; ModifierIdx >= 0; ModifierIdx--)
 	{
 		FActorModifierCoreStackRemoveOp RemoveOp;
 		RemoveOp.RemoveModifier = Modifiers[ModifierIdx];
+
 		RemoveModifier(RemoveOp);
 	}
+
 	UnlockModifierExecution();
 
 	return Modifiers.IsEmpty();
@@ -1373,6 +1453,11 @@ bool UActorModifierCoreStack::RemoveAllModifiers()
 
 void UActorModifierCoreStack::OnModifierDisabled(EActorModifierCoreDisableReason InReason)
 {
+	if (IsModifierStackFrozen())
+	{
+		return;
+	}
+
 	// Restore in reverse order
 	for (int32 ModifierIdx = Modifiers.Num() - 1; ModifierIdx >= 0; ModifierIdx--)
 	{
@@ -1411,6 +1496,21 @@ void UActorModifierCoreStack::OnActorDestroyed(AActor*)
 	UninitializeModifier(EActorModifierCoreDisableReason::Destroyed);
 }
 
+void UActorModifierCoreStack::OnModifierStackFrozenChanged()
+{
+	UE_LOGF(LogActorModifierCore, Verbose, "[%ls][%ls][%ls] Modifier stack %ls"
+		, *UE::ActorModifierCore::Utilities::GetActorNameSafe(GetModifiedActor())
+		, *GetNameSafe(GetClass())
+		, *GetModifierName().ToString()
+		, IsModifierStackFrozen() ? TEXT("frozen") : TEXT("unfrozen"));
+
+	// Re-execute stack if modifier is no longer frozen
+	if (!IsModifierStackFrozen())
+	{
+		MarkModifierDirty();
+	}
+}
+
 void UActorModifierCoreStack::OnActorTransformUpdated(USceneComponent*, EUpdateTransformFlags, ETeleportType)
 {
 	if (IsModifierEnabled())
@@ -1639,6 +1739,11 @@ void UActorModifierCoreStack::OnBlueprintModifierReplaced(UActorModifierCoreBlue
 
 void UActorModifierCoreStack::RestorePreState()
 {
+	if (IsModifierStackFrozen())
+	{
+		return;
+	}
+
 	// Restore previously applied modifiers
 	ExecutionTask.Restore();
 
@@ -1652,6 +1757,11 @@ void UActorModifierCoreStack::RestorePreState()
 
 void UActorModifierCoreStack::Apply()
 {
+	if (IsModifierStackFrozen())
+	{
+		return;
+	}
+
 	// restore original state before dirty modifiers
 	RestorePreState();
 
@@ -1685,18 +1795,38 @@ void UActorModifierCoreStack::OnModifierRemoved(EActorModifierCoreDisableReason
 		Subsystem->OnModifierReplaced().RemoveAll(this);
 	}
 
-	// un-initialize inner modifiers if they are not un-initialized
-	for (const TObjectPtr<UActorModifierCoreBase>& Modifier : Modifiers)
+	if (IsModifierStackFrozen())
 	{
-		if (Modifier)
+		for (const TObjectPtr<UActorModifierCoreBase>& Modifier : Modifiers)
+		{
+			if (Modifier)
+			{
+				// Broadcast the removal directly without going through uninit process
+				UActorModifierCoreStack::OnModifierRemovedDelegate.Broadcast(Modifier, InReason);
+			}
+		}
+	}
+	else
+	{
+		// un-initialize inner modifiers if they are not un-initialized
+		for (const TObjectPtr<UActorModifierCoreBase>& Modifier : Modifiers)
 		{
-			Modifier->UninitializeModifier(InReason);
+			if (Modifier)
+			{
+				Modifier->UninitializeModifier(InReason);
+			}
 		}
 	}
+
 }
 
 void UActorModifierCoreStack::OnModifiedActorTransformed()
 {
+	if (IsModifierStackFrozen())
+	{
+		return;
+	}
+
 	for (const TObjectPtr<UActorModifierCoreBase>& Modifier : Modifiers)
 	{
 		// only tell enable modifier that our actor has moved
@@ -1744,6 +1874,7 @@ void UActorModifierCoreStack::OnModifierDirty(UActorModifierCoreBase* DirtyModif
 		, *DirtyModifier->GetModifierName().ToString());
 
 	if (!bExecute
+		|| IsModifierStackFrozen()
 		|| DirtyModifier->IsModifierExecutionLocked()
 		|| !IsModifierStackInitialized()
 		|| !IsModifierDirty()
@@ -1774,7 +1905,8 @@ void UActorModifierCoreStack::OnModifierCDOSetup(FActorModifierCoreMetadata& InM
 
 void UActorModifierCoreStack::TickModifier(float InDelta) const
 {
-	if (!IsRootStack()
+	if (IsModifierStackFrozen() 
+		|| !IsRootStack()
 		|| !IsModifierEnabled()
 		|| !IsModifierInitialized()
 		|| !IsModifierIdle()
diff --git a/Engine/Plugins/VirtualProduction/ActorModifierCore/Source/ActorModifierCore/Public/Modifiers/ActorModifierCoreStack.h b/Engine/Plugins/VirtualProduction/ActorModifierCore/Source/ActorModifierCore/Public/Modifiers/ActorModifierCoreStack.h
--- a/Engine/Plugins/VirtualProduction/ActorModifierCore/Source/ActorModifierCore/Public/Modifiers/ActorModifierCoreStack.h
+++ b/Engine/Plugins/VirtualProduction/ActorModifierCore/Source/ActorModifierCore/Public/Modifiers/ActorModifierCoreStack.h
@@ -54,6 +54,15 @@ public:
 	/** Create a new stack by passing the actor and the parent stack if there is one */
 	static UActorModifierCoreStack* Create(UActorModifierCoreComponent* InComponent, UActorModifierCoreStack* InParentStack = nullptr);
 
+	/** Returns true if this stack is frozen */
+	ACTORMODIFIERCORE_API bool IsModifierStackFrozen() const;
+
+	/**
+	 * Set whether the modifier stack is frozen
+	 * @see bModifierStackFrozen
+	 */
+	ACTORMODIFIERCORE_API void SetModifierStackFrozen(bool bInFrozen);
+
 	/** Gets all modifiers in this stack, does not recurse */
 	TConstArrayView<UActorModifierCoreBase*> GetModifiers() const
 	{
@@ -138,6 +147,9 @@ public:
 protected:
 	//~ Begin UObject
 	virtual void PostLoad() override;
+#if WITH_EDITOR
+	virtual void PostEditChangeProperty(FPropertyChangedEvent& InPropertyChangedEvent) override;
+#endif
 	//~ End UObject
 
 	/** Process a function through each modifier in the stack and also the stacks below, stop when we return false */
@@ -219,6 +231,9 @@ private:
 	UFUNCTION()
 	void OnActorDestroyed(AActor* InActor);
 
+	/** Called when the modifier stack frozen has changed */
+	void OnModifierStackFrozenChanged();
+
 	/** Called when the actor transform is updated */
 	void OnActorTransformUpdated(USceneComponent*, EUpdateTransformFlags, ETeleportType);
 
@@ -245,6 +260,16 @@ private:
 	UPROPERTY(Transient, DuplicateTransient, NonTransactional)
 	TOptional<FActorModifierCoreExecutionTask> RootExecutionTask;
 
+	/**
+	 * EXPERIMENTAL:
+	 * Freeze the stack to disable modifiers from being added/removed/moved, while also
+	 * preventing restore states from occurring.
+	 * The only operation allowed in a frozen state is to remove the stack as a whole.
+	 * This can be used to temporarily apply a set of modifiers once, and persist the changes after removing the modifiers.
+	 */
+	UPROPERTY(EditInstanceOnly, Category="Modifier", meta=(DisplayName="Freeze Stack"))
+	bool bModifierStackFrozen = false;
+
 	/** Enable profiling for the modifiers in this stack */
 	UPROPERTY(EditInstanceOnly, Category="Modifier", meta=(DisplayName="Enable Profiling"))
 	bool bModifierProfiling = false;
diff --git a/Engine/Plugins/VirtualProduction/ActorModifierCore/Source/ActorModifierCoreEditor/Private/ActorModifierCoreEditorStyle.cpp b/Engine/Plugins/VirtualProduction/ActorModifierCore/Source/ActorModifierCoreEditor/Private/ActorModifierCoreEditorStyle.cpp
--- a/Engine/Plugins/VirtualProduction/ActorModifierCore/Source/ActorModifierCoreEditor/Private/ActorModifierCoreEditorStyle.cpp
+++ b/Engine/Plugins/VirtualProduction/ActorModifierCore/Source/ActorModifierCoreEditor/Private/ActorModifierCoreEditorStyle.cpp
@@ -24,6 +24,7 @@ FActorModifierCoreEditorStyle::FActorModifierCoreEditorStyle()
 	Set("ClassIcon.ActorModifierCoreComponent",                new IMAGE_BRUSH_SVG("Icons/BaseModifier", Icon16x16));
 	Set("ClassIcon.ActorModifierCoreEditorStackCustomization", new IMAGE_BRUSH_SVG("Icons/BaseModifier", Icon16x16));
 	Set("Profiling",                                           new IMAGE_BRUSH_SVG("Icons/Profiling",    Icon16x16));
+	Set("FreezeStack",                                         new IMAGE_BRUSH("Icons/FreezeStack",      Icon16x16));
 
 	FSlateStyleRegistry::RegisterSlateStyle(*this);
 }
diff --git a/Engine/Plugins/VirtualProduction/ActorModifierCore/Source/ActorModifierCoreEditor/Private/Modifiers/Customizations/ActorModifierCoreEditorDetailCustomization.cpp b/Engine/Plugins/VirtualProduction/ActorModifierCore/Source/ActorModifierCoreEditor/Private/Modifiers/Customizations/ActorModifierCoreEditorDetailCustomization.cpp
--- a/Engine/Plugins/VirtualProduction/ActorModifierCore/Source/ActorModifierCoreEditor/Private/Modifiers/Customizations/ActorModifierCoreEditorDetailCustomization.cpp
+++ b/Engine/Plugins/VirtualProduction/ActorModifierCore/Source/ActorModifierCoreEditor/Private/Modifiers/Customizations/ActorModifierCoreEditorDetailCustomization.cpp
@@ -24,6 +24,11 @@ void FActorModifierCoreEditorDetailCustomization::CustomizeDetails(IDetailLayout
 		UActorModifierCoreBase::StaticClass()
 	);
 
+	const TSharedPtr<IPropertyHandle> FreezePropertyHandle = InDetailBuilder.GetProperty(
+		GET_MEMBER_NAME_CHECKED(UActorModifierCoreStack, bModifierStackFrozen),
+		UActorModifierCoreStack::StaticClass()
+	);
+
 	const TSharedPtr<IPropertyHandle> ProfilingPropertyHandle = InDetailBuilder.GetProperty(
 		GET_MEMBER_NAME_CHECKED(UActorModifierCoreStack, bModifierProfiling),
 		UActorModifierCoreStack::StaticClass()
@@ -37,6 +42,7 @@ void FActorModifierCoreEditorDetailCustomization::CustomizeDetails(IDetailLayout
 	
 	if (!ModifiersPropertyHandle.IsValid()
 		|| !EnablePropertyHandle.IsValid()
+		|| !FreezePropertyHandle.IsValid()
 		|| !ProfilingPropertyHandle.IsValid()
 		|| !KeyframeHandler.IsValid()
 		|| !OperatorStackSubsystem)
@@ -45,6 +51,7 @@ void FActorModifierCoreEditorDetailCustomization::CustomizeDetails(IDetailLayout
 	}
 
 	InDetailBuilder.HideProperty(EnablePropertyHandle);
+	InDetailBuilder.HideProperty(FreezePropertyHandle);
 	InDetailBuilder.HideProperty(ProfilingPropertyHandle);
 	
 	IDetailPropertyRow* ModifiersRow = InDetailBuilder.EditDefaultProperty(ModifiersPropertyHandle);
diff --git a/Engine/Plugins/VirtualProduction/ActorModifierCore/Source/ActorModifierCoreEditor/Private/Modifiers/Customizations/ActorModifierCoreEditorStackCustomization.cpp b/Engine/Plugins/VirtualProduction/ActorModifierCore/Source/ActorModifierCoreEditor/Private/Modifiers/Customizations/ActorModifierCoreEditorStackCustomization.cpp
--- a/Engine/Plugins/VirtualProduction/ActorModifierCore/Source/ActorModifierCoreEditor/Private/Modifiers/Customizations/ActorModifierCoreEditorStackCustomization.cpp
+++ b/Engine/Plugins/VirtualProduction/ActorModifierCore/Source/ActorModifierCoreEditor/Private/Modifiers/Customizations/ActorModifierCoreEditorStackCustomization.cpp
@@ -27,6 +27,17 @@
 
 DEFINE_LOG_CATEGORY_STATIC(LogActorModifierCoreEditorStackCustomization, Log, All);
 
+namespace UE::ActorModifierCoreEditor::Private
+{
+
+FText GetStackPropertyTooltip(FName InPropertyName)
+{
+	const FProperty* Property = FindFProperty<FProperty>(UActorModifierCoreStack::StaticClass(), InPropertyName);
+	return Property ? Property->GetToolTipText() : FText::GetEmpty();
+}
+
+} // UE::ActorModifierCoreEditor::Private
+
 UActorModifierCoreEditorStackCustomization::UActorModifierCoreEditorStackCustomization()
 	: UOperatorStackEditorStackCustomization(
 		TEXT("Modifiers")
@@ -341,6 +352,7 @@ void UActorModifierCoreEditorStackCustomization::CustomizeItemBody(const FOperat
 	if (InItem->IsA<UActorModifierCoreBase>())
 	{
 		FBoolProperty* ModifierEnableProperty = FindFProperty<FBoolProperty>(UActorModifierCoreBase::StaticClass(), GET_MEMBER_NAME_CHECKED(UActorModifierCoreBase, bModifierEnabled));
+		FBoolProperty* StackFrozenProperty = FindFProperty<FBoolProperty>(UActorModifierCoreStack::StaticClass(), GET_MEMBER_NAME_CHECKED(UActorModifierCoreStack, bModifierStackFrozen));
 		FBoolProperty* ProfilingEnableProperty = FindFProperty<FBoolProperty>(UActorModifierCoreStack::StaticClass(), GET_MEMBER_NAME_CHECKED(UActorModifierCoreStack, bModifierProfiling));
 
 		const bool bIsStack = InItem->IsA<UActorModifierCoreStack>();
@@ -348,6 +360,7 @@ void UActorModifierCoreEditorStackCustomization::CustomizeItemBody(const FOperat
 		InBodyBuilder
 			.SetShowDetailsView(!bIsStack)
 			.DisallowProperty(ModifierEnableProperty)
+			.DisallowProperty(StackFrozenProperty)
 			.DisallowProperty(ProfilingEnableProperty);
 	}
 
@@ -555,6 +568,8 @@ void UActorModifierCoreEditorStackCustomization::FillStackHeaderMenu(UToolMenu*
 
 void UActorModifierCoreEditorStackCustomization::FillItemHeaderActionMenu(UToolMenu* InToolMenu) const
 {
+	using namespace UE::ActorModifierCoreEditor;
+
 	if (!InToolMenu)
 	{
 		return;
@@ -572,9 +587,22 @@ void UActorModifierCoreEditorStackCustomization::FillItemHeaderActionMenu(UToolM
 		return;
 	}
 
-	// Add profiling stat toggle entry
+	// Add freeze & profiling stat toggle entry
 	if (ItemContext->IsA<UActorModifierCoreStack>())
 	{
+		const FToolMenuEntry FreezeModifierStackAction = FToolMenuEntry::InitToolBarButton(
+			TEXT("FreezeModifierStackMenuEntry")
+			, FUIAction(
+				FExecuteAction::CreateUObject(this, &UActorModifierCoreEditorStackCustomization::ToggleModifierStackFrozenAction, ItemContext)
+				, FCanExecuteAction()
+				, FIsActionChecked::CreateUObject(this, &UActorModifierCoreEditorStackCustomization::IsModifierStackFrozen, ItemContext))
+			, FText::GetEmpty()
+			, Private::GetStackPropertyTooltip(GET_MEMBER_NAME_CHECKED(UActorModifierCoreStack, bModifierStackFrozen))
+			, FSlateIcon(FActorModifierCoreEditorStyle::Get().GetStyleSetName(), "FreezeStack")
+			, EUserInterfaceActionType::ToggleButton
+		);
+		InToolMenu->AddMenuEntry(FreezeModifierStackAction.Name, FreezeModifierStackAction);
+
 		const FToolMenuEntry EnableProfilingModifierAction = FToolMenuEntry::InitToolBarButton(
 			TEXT("EnableProfilingModifierMenuEntry")
 			, FUIAction(
@@ -607,6 +635,8 @@ void UActorModifierCoreEditorStackCustomization::FillItemHeaderActionMenu(UToolM
 
 void UActorModifierCoreEditorStackCustomization::FillItemContextActionMenu(UToolMenu* InToolMenu) const
 {
+	using namespace UE::ActorModifierCoreEditor;
+
 	if (!InToolMenu)
 	{
 		return;
@@ -626,6 +656,20 @@ void UActorModifierCoreEditorStackCustomization::FillItemContextActionMenu(UTool
 
 	if (ItemContext->IsA<UActorModifierCoreStack>())
 	{
+		const FToolMenuEntry FreezeModifierStackAction = FToolMenuEntry::InitMenuEntry(
+			TEXT("FreezeModifierStackMenuEntry")
+			, LOCTEXT("FreezeModifierStack", "Freeze Modifier Stack")
+			, Private::GetStackPropertyTooltip(GET_MEMBER_NAME_CHECKED(UActorModifierCoreStack, bModifierStackFrozen))
+			, FSlateIcon(FActorModifierCoreEditorStyle::Get().GetStyleSetName(), "FreezeStack")
+			, FUIAction(
+				FExecuteAction::CreateUObject(this, &UActorModifierCoreEditorStackCustomization::ToggleModifierStackFrozenAction, ItemContext)
+				, FCanExecuteAction()
+				, FIsActionChecked::CreateUObject(this, &UActorModifierCoreEditorStackCustomization::IsModifierStackFrozen, ItemContext))
+			, EUserInterfaceActionType::ToggleButton
+		);
+
+		InToolMenu->AddMenuEntry(FreezeModifierStackAction.Name, FreezeModifierStackAction);
+
 		const FToolMenuEntry EnableProfilingModifierAction = FToolMenuEntry::InitMenuEntry(
 			TEXT("EnableProfilingModifierMenuEntry")
 			, LOCTEXT("EnableProfilingModifier", "Toggle profiling")
@@ -675,6 +719,20 @@ bool UActorModifierCoreEditorStackCustomization::CanRemoveModifier(FOperatorStac
 		return false;
 	}
 
+	for (UActorModifierCoreBase* Modifier : InItem->GetAsArray<UActorModifierCoreBase>())
+	{
+		if (!Modifier)
+		{
+			continue;
+		}
+
+		const UActorModifierCoreStack* const ModifierStack = Modifier->GetModifierStack();
+		if (ModifierStack && ModifierStack->IsModifierStackFrozen())
+		{
+			return false;
+		}
+	}
+
 	return true;
 }
 
@@ -855,6 +913,41 @@ void UActorModifierCoreEditorStackCustomization::PasteModifierAction(FOperatorSt
 	}
 }
 
+bool UActorModifierCoreEditorStackCustomization::IsModifierStackFrozen(FOperatorStackEditorItemPtr InItem) const
+{
+	if (!InItem.IsValid() || !InItem->HasValue())
+	{
+		return false;
+	}
+
+	for (const UActorModifierCoreStack* ModifierStack : InItem->GetAsArray<UActorModifierCoreStack>())
+	{
+		if (ModifierStack && ModifierStack->IsModifierStackFrozen())
+		{
+			return true;
+		}
+	}
+
+	return false;
+}
+
+void UActorModifierCoreEditorStackCustomization::ToggleModifierStackFrozenAction(FOperatorStackEditorItemPtr InItem) const
+{
+	if (!InItem.IsValid() || !InItem->HasValue())
+	{
+		return;
+	}
+
+	const bool bIsModifierFrozen = IsModifierStackFrozen(InItem);
+	for (UActorModifierCoreStack* ModifierStack : InItem->GetAsArray<UActorModifierCoreStack>())
+	{
+		if (ModifierStack)
+		{
+			ModifierStack->SetModifierStackFrozen(!bIsModifierFrozen);	
+		}
+	}
+}
+
 bool UActorModifierCoreEditorStackCustomization::IsModifierProfiling(FOperatorStackEditorItemPtr InItem) const
 {
 	if (!InItem.IsValid() || !InItem->HasValue())
diff --git a/Engine/Plugins/VirtualProduction/ActorModifierCore/Source/ActorModifierCoreEditor/Private/Modifiers/Customizations/ActorModifierCoreEditorStackCustomization.h b/Engine/Plugins/VirtualProduction/ActorModifierCore/Source/ActorModifierCoreEditor/Private/Modifiers/Customizations/ActorModifierCoreEditorStackCustomization.h
--- a/Engine/Plugins/VirtualProduction/ActorModifierCore/Source/ActorModifierCoreEditor/Private/Modifiers/Customizations/ActorModifierCoreEditorStackCustomization.h
+++ b/Engine/Plugins/VirtualProduction/ActorModifierCore/Source/ActorModifierCoreEditor/Private/Modifiers/Customizations/ActorModifierCoreEditorStackCustomization.h
@@ -73,6 +73,10 @@ protected:
 	bool CanPasteModifier(FOperatorStackEditorItemPtr InItem) const;
 	void PasteModifierAction(FOperatorStackEditorItemPtr InItem) const;
 
+	/** Toggle the frozen state for a whole stack */
+	bool IsModifierStackFrozen(FOperatorStackEditorItemPtr InItem) const;
+	void ToggleModifierStackFrozenAction(FOperatorStackEditorItemPtr InItem) const;
+
 	/** Toggle the profiling for a whole stack */
 	bool IsModifierProfiling(FOperatorStackEditorItemPtr InItem) const;
 	void ToggleModifierProfilingAction(FOperatorStackEditorItemPtr InItem) const;
